FAB plugin - add pre-release source

diff --git a/Engine/Plugins/Fab/Fab.uplugin b/Engine/Plugins/Fab/Fab.uplugin
new file mode 100644
--- /dev/null
+++ b/Engine/Plugins/Fab/Fab.uplugin
@@ -0,0 +1,23 @@
+{
+	"FileVersion": 3,
+	"Version": 1,
+	"VersionName": "1.0",
+	"FriendlyName": "Fab",
+	"Description": "Pre-release Fab Plugin",
+	"Category": "Other",
+	"CreatedBy": "Epic Games, Inc.",
+	"CreatedByURL": "http://epicgames.com",
+	"DocsURL": "",
+	"MarketplaceURL": "",
+	"SupportURL": "",
+	"CanContainContent": true,
+	"IsBetaVersion": false,
+	"IsExperimentalVersion": false,
+	"EnabledByDefault": true,
+	"Modules": [
+		{
+			"Name": "Fab",
+			"Type": "Editor"
+		}
+	]
+}
\ No newline at end of file
diff --git a/Engine/Plugins/Fab/Source/Fab/Fab.Build.cs b/Engine/Plugins/Fab/Source/Fab/Fab.Build.cs
new file mode 100644
--- /dev/null
+++ b/Engine/Plugins/Fab/Source/Fab/Fab.Build.cs
@@ -0,0 +1,68 @@
+// Copyright Epic Games, Inc. All Rights Reserved.
+
+using UnrealBuildTool;
+
+public class Fab : ModuleRules
+{
+	public Fab(ReadOnlyTargetRules Target) : base(Target)
+	{
+		PCHUsage = ModuleRules.PCHUsageMode.UseExplicitOrSharedPCHs;
+		
+		PublicIncludePaths.AddRange(
+			new string[] 
+			{
+				// ... add public include paths required here ...
+			}
+		);
+				
+		
+		PrivateIncludePaths.AddRange(
+			new string[] 
+			{
+				// ... add other private include paths required here ...
+			}
+		);
+			
+		
+		PublicDependencyModuleNames.AddRange(
+			new string[]
+			{
+				"Core",
+				// ... add other public dependencies that you statically link with here ...
+			}
+		);
+			
+		
+		PrivateDependencyModuleNames.AddRange(
+			new string[]
+			{
+				"CoreUObject",
+				"Engine",
+				"Slate",
+				"SlateCore",
+				"WebBrowser",
+				"HTTP",
+				"MainFrame",
+				"EditorStyle",
+				"EditorSubsystem",
+				"ToolMenus",
+				"LevelEditor",
+				"Projects",
+				"PlacementMode",
+				"ApplicationCore",
+				"UnrealEd",
+				"InputCore",
+				"RenderCore",
+				// ... add private dependencies that you statically link with here ...	
+			}
+		);
+		
+		
+		DynamicallyLoadedModuleNames.AddRange(
+			new string[]
+			{
+				// ... add any modules that your module loads dynamically here ...
+			}
+		);
+	}
+}
diff --git a/Engine/Plugins/Fab/Source/Fab/Private/FabModule.cpp b/Engine/Plugins/Fab/Source/Fab/Private/FabModule.cpp
new file mode 100644
--- /dev/null
+++ b/Engine/Plugins/Fab/Source/Fab/Private/FabModule.cpp
@@ -0,0 +1,17 @@
+// Copyright Epic Games, Inc. All Rights Reserved.
+
+#include "FabModule.h"
+
+class FFabModule : public IFabModule
+{
+public:
+	virtual void StartupModule() override
+	{
+	}
+
+	virtual void ShutdownModule() override
+	{
+	}
+};
+
+IMPLEMENT_MODULE(FFabModule, Fab);
diff --git a/Engine/Plugins/Fab/Source/Fab/Public/FabModule.h b/Engine/Plugins/Fab/Source/Fab/Public/FabModule.h
new file mode 100644
--- /dev/null
+++ b/Engine/Plugins/Fab/Source/Fab/Public/FabModule.h
@@ -0,0 +1,21 @@
+// Copyright Epic Games, Inc. All Rights Reserved.
+
+#pragma once
+
+#include "Modules/ModuleManager.h"
+
+#define MS_MODULE_NAME TEXT("Fab")
+
+class IFabModule : public IModuleInterface
+{
+public:
+	static inline IFabModule& Get()
+	{
+		return FModuleManager::LoadModuleChecked<IFabModule>(MS_MODULE_NAME);
+	}
+
+	static inline bool IsAvailable()
+	{
+		return FModuleManager::Get().IsModuleLoaded(MS_MODULE_NAME);
+	}
+};
