[Fab] Potential fix for plant material warnings in CI

diff --git a/Engine/Plugins/Fab/Source/Fab/Private/FabModule.cpp b/Engine/Plugins/Fab/Source/Fab/Private/FabModule.cpp
--- a/Engine/Plugins/Fab/Source/Fab/Private/FabModule.cpp
+++ b/Engine/Plugins/Fab/Source/Fab/Private/FabModule.cpp
@@ -32,11 +32,8 @@ class FFabModule : public IFabModule
 public:
 	virtual void StartupModule() override
 	{
-		if (GIsEditor && !IsRunningCommandlet())
+		if (GIsEditor)
 		{
-			FFabBrowser::Init();
-			FabAuthentication::Init();
-
 			URendererSettings* RendererSettings                = GetMutableDefault<URendererSettings>();
 			RendererSettings->bEnableVirtualTextureOpacityMask = true;
 			RendererSettings->PostEditChange();
@@ -48,6 +45,12 @@ public:
 				EditorSettings->PostEditChange();
 			}
 #endif
+		}
+		
+		if (GIsEditor && !IsRunningCommandlet())
+		{
+			FFabBrowser::Init();
+			FabAuthentication::Init();
 
 			FPropertyEditorModule& PropertyModule = FModuleManager::LoadModuleChecked<FPropertyEditorModule>("PropertyEditor");
 			PropertyModule.RegisterCustomClassLayout("FabSettings", FOnGetDetailCustomizationInstance::CreateStatic(&FFabSettingsCustomization::MakeInstance));
