[Fab] Use platform allow list for the module and interchange plugin

diff --git a/Engine/Plugins/Fab/Fab.uplugin b/Engine/Plugins/Fab/Fab.uplugin
--- a/Engine/Plugins/Fab/Fab.uplugin
+++ b/Engine/Plugins/Fab/Fab.uplugin
@@ -1,7 +1,7 @@
 {
 	"FileVersion": 3,
-	"Version": 1,
-	"VersionName": "0.0.1",
+	"Version": 3,
+	"VersionName": "0.0.3",
 	"FriendlyName": "Fab",
 	"Description": "Fab Plugin",
 	"Category": "Other",
@@ -12,15 +12,16 @@
 	"SupportURL": "",
 	"CanContainContent": true,
 	"EnabledByDefault": true,
-	"SupportedTargetPlatforms": [
-		"Win64",
-		"Mac",
-		"Linux"
-	],
 	"Plugins": [
 		{
 			"Name": "Interchange",
-			"Enabled": true
+			"Enabled": true,
+			"PlatformAllowList" :
+			[
+				"Linux",
+				"Win64",
+				"Mac"
+			]
 		},
 		{
 			"Name": "EditorScriptingUtilities",
@@ -35,7 +36,13 @@
 		{
 			"Name": "Fab",
 			"Type": "Editor",
-			"LoadingPhase": "PostDefault"
+			"LoadingPhase": "PostDefault",
+			"PlatformAllowList" :
+			[
+				"Linux",
+				"Win64",
+				"Mac"
+			]
 		}
 	]
 }
