[Fab] Potential fix for unzip vulnerability that leads to file write on user disk

diff --git a/Engine/Plugins/Fab/Source/Fab/Private/Utilities/AssetUtils.cpp b/Engine/Plugins/Fab/Source/Fab/Private/Utilities/AssetUtils.cpp
--- a/Engine/Plugins/Fab/Source/Fab/Private/Utilities/AssetUtils.cpp
+++ b/Engine/Plugins/Fab/Source/Fab/Private/Utilities/AssetUtils.cpp
@@ -54,6 +54,9 @@ bool FAssetUtils::Unzip(const FString& Path, const FString& TargetPath)
 	{
 		if (FileName.EndsWith("/") || FileName.EndsWith("\\"))
 			continue;
+		const FString AbsoluteDestFileName = TargetPath / FileName;
+		if (!FPaths::IsUnderDirectory(AbsoluteDestFileName, TargetPath))
+			continue;
 		if (TArray<uint8> FileBuffer; ZipArchiveReader.TryReadFile(FileName, FileBuffer))
 		{
 			if (!FFileHelper::SaveArrayToFile(FileBuffer, *(TargetPath / FileName)))
