[Fab] Initial MyFolder integration into TEDS.

This CR presents the start of integrating Fab's "My Folder" into the Content Browser. In this first step a new CVar command "Fab.TEDS.MyFolderIntegration" was added that downloads the content of My Folder, parses it and stores it in the Editor Data Storage. The TEDS Debugger with the "FFabObjectColumn" selected can show the results. The command can safely be repeated and previous entries will automatically be removed first. The Fab default of downloading batches of a 1000 results at a time is used but can be adjusted by giving a different number with the CVar command. The same settings as for the for Fab Browser are used and can be changed using the other Fab CVar commands.

[RN] [Fab] Initial MyFolder integration into TEDS.

diff --git a/Engine/Plugins/Fab/Fab.uplugin b/Engine/Plugins/Fab/Fab.uplugin
--- a/Engine/Plugins/Fab/Fab.uplugin
+++ b/Engine/Plugins/Fab/Fab.uplugin
@@ -23,6 +23,10 @@
 				"Mac"
 			]
 		},
+		{
+			"Name": "EditorDataStorageFeatures",
+			"Enabled": true
+		},
 		{
 			"Name": "EditorScriptingUtilities",
 			"Enabled": true
diff --git a/Engine/Plugins/Fab/Source/Fab/Fab.Build.cs b/Engine/Plugins/Fab/Source/Fab/Fab.Build.cs
--- a/Engine/Plugins/Fab/Source/Fab/Fab.Build.cs
+++ b/Engine/Plugins/Fab/Source/Fab/Fab.Build.cs
@@ -67,6 +67,7 @@ public class Fab : ModuleRules
 				"SlateCore",
 				"ToolMenus",
 				"ToolWidgets",
+				"TypedElementFramework",
 				"UMG",
 				"UnrealEd",
 				"WebBrowser",
diff --git a/Engine/Plugins/Fab/Source/Fab/Private/FabBrowser.cpp b/Engine/Plugins/Fab/Source/Fab/Private/FabBrowser.cpp
--- a/Engine/Plugins/Fab/Source/Fab/Private/FabBrowser.cpp
+++ b/Engine/Plugins/Fab/Source/Fab/Private/FabBrowser.cpp
@@ -235,34 +235,26 @@ void FFabBrowser::RegisterNomadTab()
 
 FString FFabBrowser::GetUrl()
 {
-	const FString ProdUrl = TEXT("https://www.fab.com/plugins/ue5");
 	if (FabPluginSettings == nullptr)
 	{
-		return ProdUrl;
+		return TEXT("https://www.fab.com/plugins/ue5");
 	}
 
 	switch (FabPluginSettings->Environment)
 	{
-		case EFabEnvironment::Prod:
+		default: // fall through
+		case EFabEnvironment::Prod:  // fall through
+		case EFabEnvironment::Gamedev: // fall through
+		case EFabEnvironment::Test:  // fall through
 		{
-			return ProdUrl;
-		}
-		case EFabEnvironment::Gamedev:
-		{
-			return TEXT("https://fab.cceb.dev.use1a.on.epicgames.com/plugins/ue5");
-		}
-		case EFabEnvironment::Test:
-		{
-			return TEXT("https://fab.daec.live.use1a.on.epicgames.com/plugins/ue5");
+			FString Url = FabPluginSettings->GetUrlFromEnvironment();
+			Url += TEXT("/plugins/ue5");
+			return Url;
 		}
 		case EFabEnvironment::CustomUrl:
 		{
 			return FabPluginSettings->CustomUrl;
 		}
-		default:
-		{
-			return ProdUrl;
-		}
 	}
 }
 
diff --git a/Engine/Plugins/Fab/Source/Fab/Private/FabConsoleCommands.cpp b/Engine/Plugins/Fab/Source/Fab/Private/FabConsoleCommands.cpp
--- a/Engine/Plugins/Fab/Source/Fab/Private/FabConsoleCommands.cpp
+++ b/Engine/Plugins/Fab/Source/Fab/Private/FabConsoleCommands.cpp
@@ -4,6 +4,7 @@
 #include "FabBrowser.h"
 #include "FabLog.h"
 #include "FabSettings.h"
+#include "Teds/FabMyFolderIntegration.h"
 
 #include "HAL/IConsoleManager.h"
 #include "Utilities/FabAssetsCache.h"
@@ -68,3 +69,21 @@ static FAutoConsoleCommand ConsoleCmd_FabSetEnvironment(
 		FabSettings->SaveConfig();
 	})
 );
+
+static FAutoConsoleCommand ConsoleCmd_FabTedsMyFolderIntegration(
+	TEXT("Fab.TEDS.MyFolderIntegration"),
+	TEXT("Loads the content of My Folder into TEDS for use in the editor, e.g. by the Content Browser."),
+	FConsoleCommandWithArgsDelegate::CreateLambda(
+		[](const TArray<FString>& Args)
+		{
+			if (!Args.IsEmpty())
+			{
+				if (int32 BatchSize = FCString::Atoi(*(Args[0])); BatchSize > 0)
+				{
+					FFabTedsMyFolderIntegration::QueueSyncRequest(static_cast<uint32>(BatchSize));
+					return;
+				}
+			}
+			FFabTedsMyFolderIntegration::QueueSyncRequest();
+		})
+);
diff --git a/Engine/Plugins/Fab/Source/Fab/Private/FabSettings.cpp b/Engine/Plugins/Fab/Source/Fab/Private/FabSettings.cpp
--- a/Engine/Plugins/Fab/Source/Fab/Private/FabSettings.cpp
+++ b/Engine/Plugins/Fab/Source/Fab/Private/FabSettings.cpp
@@ -44,3 +44,20 @@ void UFabSettings::PostEditChangeProperty(struct FPropertyChangedEvent& Property
 		}
 	}
 }
+
+FString UFabSettings::GetUrlFromEnvironment() const
+{
+	switch (Environment)
+	{
+	default:
+		// fall through
+	case EFabEnvironment::Prod:
+		return TEXT("https://www.fab.com");
+	case EFabEnvironment::Gamedev:
+		return TEXT("https://fab.cceb.dev.use1a.on.epicgames.com");
+	case EFabEnvironment::Test:
+		return TEXT("https://fab.daec.live.use1a.on.epicgames.com");
+	case EFabEnvironment::CustomUrl:
+		return FString();
+	}
+}
diff --git a/Engine/Plugins/Fab/Source/Fab/Private/FabSettings.h b/Engine/Plugins/Fab/Source/Fab/Private/FabSettings.h
--- a/Engine/Plugins/Fab/Source/Fab/Private/FabSettings.h
+++ b/Engine/Plugins/Fab/Source/Fab/Private/FabSettings.h
@@ -74,4 +74,6 @@ public:
 	EFabPreferredQualityTier PreferredQualityTier = EFabPreferredQualityTier::Medium;
 
 	virtual void PostEditChangeProperty(struct FPropertyChangedEvent& PropertyChangedEvent) override;
+
+	FString GetUrlFromEnvironment() const;
 };
diff --git a/Engine/Plugins/Fab/Source/Fab/Private/Teds/FabMyFolderIntegration.cpp b/Engine/Plugins/Fab/Source/Fab/Private/Teds/FabMyFolderIntegration.cpp
new file mode 100644
--- /dev/null
+++ b/Engine/Plugins/Fab/Source/Fab/Private/Teds/FabMyFolderIntegration.cpp
@@ -0,0 +1,284 @@
+// Copyright Epic Games, Inc. All Rights Reserved.
+
+#include "Teds/FabMyFolderIntegration.h"
+
+#include "Dom/JsonObject.h"
+#include "Elements/Columns/TypedElementMiscColumns.h"
+#include "Elements/Columns/TypedElementWebColumns.h"
+#include "Elements/Common/EditorDataStorageFeatures.h"
+#include "Elements/Interfaces/TypedElementDataStorageInterface.h"
+#include "EOSShared.h"
+#include "FabAuthentication.h"
+#include "FabLog.h"
+#include "FabSettings.h"
+#include "HttpModule.h"
+#include "Interfaces/IHttpRequest.h"
+#include "Interfaces/IHttpResponse.h"
+#include "Serialization/JsonReader.h"
+#include "Serialization/JsonSerializer.h"
+
+void UFabFactory::RegisterTables(IEditorDataStorageProvider& DataStorage)
+{
+	Table = DataStorage.RegisterTable<FNameColumn, FFabObjectColumn, FEditorDataStorageUrlColumn>(FName("Fab"));
+}
+
+void FFabTedsMyFolderIntegration::QueueSyncRequest()
+{
+	QueueSyncRequest(1000);
+}
+
+void FFabTedsMyFolderIntegration::QueueSyncRequest(uint32 BatchSize)
+{
+	using namespace UE::Editor::DataStorage;
+
+	if (IEditorDataStorageProvider* Storage = GetMutableDataStorageFeature<IEditorDataStorageProvider>(StorageFeatureName))
+	{
+		Storage->RemoveAllRowsWith<FFabObjectColumn>();
+
+		QueueSyncRequest({}, BatchSize);
+	}
+}
+
+void FFabTedsMyFolderIntegration::QueueSyncRequest(const FString& Cursor, uint32 BatchSize)
+{
+	if (EOS_Auth_GetLoggedInAccountsCount(FabAuthentication::AuthHandle) > 0)
+	{
+		const UFabSettings* FabPluginSettings = GetDefault<UFabSettings>();
+
+		EOS_Auth_Token* UserAuthToken = nullptr;
+
+		EOS_Auth_CopyUserAuthTokenOptions CopyTokenOptions = { 0 };
+		CopyTokenOptions.ApiVersion = EOS_AUTH_COPYUSERAUTHTOKEN_API_LATEST;
+
+		if (EOS_Auth_CopyUserAuthToken(FabAuthentication::AuthHandle,
+			&CopyTokenOptions, FabAuthentication::EpicAccountId, &UserAuthToken) == EOS_EResult::EOS_Success)
+		{
+			FString CursorString;
+			if (!Cursor.IsEmpty())
+			{
+				CursorString = "&cursor=\"";
+				CursorString += Cursor;
+				CursorString += '"';
+			}
+
+			TSharedRef<IHttpRequest> HttpRequest = FHttpModule::Get().CreateRequest();
+			HttpRequest->SetVerb("GET");
+			HttpRequest->SetHeader("accept", "application/json");
+			HttpRequest->SetHeader("Authorization", FString::Printf(TEXT("Bearer %s"), *FString(UserAuthToken->AccessToken)));
+			HttpRequest->SetURL(FString::Printf(TEXT("%s/e/accounts/%s/ue/library?count=%u%s"),
+				*FabPluginSettings->GetUrlFromEnvironment(),
+				*LexToString(FabAuthentication::EpicAccountId), BatchSize, *CursorString));
+			HttpRequest->OnProcessRequestComplete().BindLambda([BatchSize](HttpRequestPtr Request, HttpResponsePtr Response, bool bWasSuccessful)
+				{
+					ProcessSyncResults(BatchSize, Request, Response, bWasSuccessful);
+				});
+			HttpRequest->ProcessRequest();
+
+			EOS_Auth_Token_Release(UserAuthToken);
+		}
+	}
+	else
+	{
+		FAB_LOG("Unable to retrieve My Folder data due to user not being logged into Fab.");
+	}
+}
+
+void FFabTedsMyFolderIntegration::ProcessSyncResults(uint32 BatchSize, HttpRequestPtr Request, HttpResponsePtr Response, bool bWasSuccessful)
+{
+	using namespace UE::Editor::DataStorage;
+
+	FAB_LOG("Result for (portion of) the My Folder data.");
+	
+	FString ContentType = bWasSuccessful ? Response->GetContentType() : FString();
+	if (bWasSuccessful && Response->GetContentType().StartsWith(TEXT("application/json")))
+	{
+		FString Content = Response->GetContentAsString();
+		TSharedPtr<FJsonObject> JsonResults = MakeShareable(new FJsonObject());
+		TSharedRef<TJsonReader<TCHAR>> JsonReader = TJsonReaderFactory<TCHAR>::Create(Content);
+		if (FJsonSerializer::Deserialize(JsonReader, JsonResults))
+		{
+			if (IEditorDataStorageProvider* Storage = GetMutableDataStorageFeature<IEditorDataStorageProvider>(StorageFeatureName))
+			{
+				if (const UFabFactory* FabFactory = Storage->FindFactory<UFabFactory>())
+				{
+					// Determine if this is the last set or that there are more results to retrieve.
+					const TSharedPtr<FJsonObject>* Cursors = nullptr;
+					if (JsonResults->TryGetObjectField(TEXT("cursors"), Cursors))
+					{
+						FString Next;
+						if ((*Cursors)->TryGetStringField(TEXT("next"), Next))
+						{
+							QueueSyncRequest(Next, BatchSize);
+						}
+					}
+
+					// Process the returned results.
+					const TArray<TSharedPtr<FJsonValue>>* Results = nullptr;
+					if (JsonResults->TryGetArrayField(TEXT("results"), Results); Results && !Results->IsEmpty())
+					{
+						FString Temp;
+
+						TArray<RowHandle> FailedSerializationRows;
+
+						int32 Index = 0;
+						Storage->BatchAddRow(FabFactory->Table, Results->Num(),
+							[Storage, Results = *Results, &Index, &Temp, &FailedSerializationRows](RowHandle Row)
+							{
+								const TSharedPtr<FJsonValue>& Result = Results[Index++];
+								
+								const TSharedPtr<FJsonObject>* ResultObject = nullptr;
+								if (bool Found = Result->TryGetObject(ResultObject); Found && ResultObject->IsValid())
+								{
+									if (FNameColumn* Target = Storage->GetColumn<FNameColumn>(Row))
+									{
+										SetNameColumn(*Target, **ResultObject, Temp);
+									}
+
+									if (FFabObjectColumn* Target = Storage->GetColumn<FFabObjectColumn>(Row))
+									{
+										SetFabObjectColumn(*Target, **ResultObject, Temp);
+									}
+
+									if (FUrlColumn* Target = Storage->GetColumn<FUrlColumn>(Row))
+									{
+										SetUrlColumn(*Target, **ResultObject, Temp);
+									}
+
+									AddDistributionMethod(Storage, Row, **ResultObject, Temp);
+									AddImages(Storage, Row, **ResultObject, Temp);
+								}
+								else
+								{
+									FailedSerializationRows.Add(Row);
+								}
+							});
+
+						if (!FailedSerializationRows.IsEmpty())
+						{
+							Storage->BatchRemoveRows(FailedSerializationRows);
+						}
+
+						FAB_LOG("Parsed data for (portion of) My Folder.");
+					}
+					else
+					{
+						FAB_LOG_ERROR("Unable to store My Folder data due missing results. An error may have occurred: %s", *Content);
+					}
+				}
+				else
+				{
+					FAB_LOG_ERROR("Unable to store My Folder data due the factory for Fab objects hasn't been initialized.");
+				}
+			}
+			else
+			{
+				FAB_LOG_ERROR("Unable to store My Folder data due the editor data storage not being available.");
+			}
+		}
+		else
+		{
+			FAB_LOG_ERROR("Unable to retrieve My Folder data due to returned result not being valid JSON.");
+		}
+	}
+	else
+	{
+		FAB_LOG_ERROR("Unable to retrieve My Folder data due to the request failing or not returning a JSON document with the required data.");
+	}
+}
+
+void FFabTedsMyFolderIntegration::SetNameColumn(FNameColumn& Target, const FJsonObject& Object, FString& Temp)
+{
+	if (Object.TryGetStringField(TEXT("title"), Temp))
+	{
+		Target.Name = FName(Temp);
+		Temp.Reset();
+	}
+}
+
+void FFabTedsMyFolderIntegration::SetFabObjectColumn(FFabObjectColumn& Target, const FJsonObject& Object, FString& Temp)
+{
+	Object.TryGetStringField(TEXT("description"), Target.Description);
+
+	if (Object.TryGetStringField(TEXT("assetId"), Temp))
+	{
+		FGuid::Parse(Temp, Target.AssetId);
+		Temp.Reset();
+	}
+	if (Object.TryGetStringField(TEXT("assetNamespace"), Temp))
+	{
+		FGuid::Parse(Temp, Target.AssetNamespace);
+		Temp.Reset();
+	}
+
+	if (Object.TryGetStringField(TEXT("listingType"), Temp))
+	{
+		Target.ListingType = FName(Temp);
+		Temp.Reset();
+	}
+
+	Object.TryGetStringField(TEXT("seller"), Target.Seller);
+
+	if (Object.TryGetStringField(TEXT("source"), Temp))
+	{
+		Target.Source = FName(Temp);
+		Temp.Reset();
+	}
+
+	/* Additionally available data :
+	* projectVersions": 
+	  [
+        {
+          "artifactId": {},
+          "engineVersions": [],
+          "targetPlatforms": [],
+		  "buildVersions": []
+        }
+      ],
+	  "customAttributes": []
+	  "legacyItemId": {}
+	*/
+}
+
+void FFabTedsMyFolderIntegration::SetUrlColumn(UE::Editor::DataStorage::FUrlColumn& Target, const FJsonObject& Object, FString& Temp)
+{
+	Object.TryGetStringField(TEXT("url"), Target.UrlString);
+}
+
+void FFabTedsMyFolderIntegration::AddDistributionMethod(
+	IEditorDataStorageProvider* Storage, UE::Editor::DataStorage::RowHandle Row, const FJsonObject& Object, FString& Temp)
+{
+	if (Object.TryGetStringField(TEXT("distributionMethod"), Temp))
+	{
+		Temp.ToLowerInline();
+		Temp.ReplaceCharInline(TEXT('_'), TEXT(' '));
+
+		Storage->AddColumn<FFabDistributionMethodTag>(Row, FName(Temp));
+		Temp.Reset();
+	}
+}
+
+void FFabTedsMyFolderIntegration::AddImages(
+	IEditorDataStorageProvider* Storage, UE::Editor::DataStorage::RowHandle Row, const FJsonObject& Object, FString& Temp)
+{
+	using namespace UE::Editor::DataStorage;
+
+	const TArray<TSharedPtr<FJsonValue>>* Images = nullptr;
+	if (Object.TryGetArrayField(TEXT("images"), Images); Images && !Images->IsEmpty())
+	{
+		for (const TSharedPtr<FJsonValue>& Image : *Images)
+		{
+			const TSharedPtr<FJsonObject>* ImageObject = nullptr;
+			if (bool Found = Image->TryGetObject(ImageObject); Found && ImageObject->IsValid())
+			{
+				FWebImageColumn ImageColumn;
+				(*ImageObject)->TryGetStringField(TEXT("url"), ImageColumn.UrlString);
+				(*ImageObject)->TryGetNumberField(TEXT("width"), ImageColumn.Width);
+				(*ImageObject)->TryGetNumberField(TEXT("height"), ImageColumn.Height);
+				
+				(*ImageObject)->TryGetStringField(TEXT("type"), Temp);
+				Temp.TrimStartAndEndInline();
+				Storage->AddColumn(Row, FName(Temp), MoveTemp(ImageColumn));
+			}
+		}
+	}
+}
diff --git a/Engine/Plugins/Fab/Source/Fab/Private/Teds/FabMyFolderIntegration.h b/Engine/Plugins/Fab/Source/Fab/Private/Teds/FabMyFolderIntegration.h
new file mode 100644
--- /dev/null
+++ b/Engine/Plugins/Fab/Source/Fab/Private/Teds/FabMyFolderIntegration.h
@@ -0,0 +1,82 @@
+// Copyright Epic Games, Inc. All Rights Reserved.
+
+#pragma once
+
+#include "Elements/Common/TypedElementCommonTypes.h"
+#include "Elements/Interfaces/TypedElementDataStorageFactory.h"
+#include "Templates/SharedPointer.h"
+
+#include "FabMyFolderIntegration.generated.h"
+
+struct FEditorDataStorageUrlColumn;
+class IEditorDataStorageProvider;
+class FJsonObject;
+struct FNameColumn;
+
+UCLASS()
+class UFabFactory : public UEditorDataStorageFactory
+{
+	GENERATED_BODY()
+
+public:
+	virtual ~UFabFactory() override = default;
+
+	virtual void RegisterTables(IEditorDataStorageProvider& DataStorage) override;
+
+	UE::Editor::DataStorage::TableHandle Table;
+};
+
+USTRUCT(meta = (DisplayName = "Fab distribution method", EditorDataStorage_DynamicColumnTemplate))
+struct FFabDistributionMethodTag : public FEditorDataStorageTag
+{
+	GENERATED_BODY()
+};
+
+USTRUCT(meta = (DisplayName = "Fab object"))
+struct FFabObjectColumn : public FEditorDataStorageColumn
+{
+	GENERATED_BODY()
+
+	UPROPERTY(meta = (Searchable))
+	FString Description;
+	
+	UPROPERTY()
+	FGuid AssetId;
+	
+	UPROPERTY()
+	FGuid AssetNamespace;
+
+	UPROPERTY(meta = (Searchable))
+	FName ListingType;
+
+	UPROPERTY(meta = (Searchable))
+	FString Seller;
+
+	UPROPERTY()
+	FName Source;
+
+	UPROPERTY()
+	FString UrlString;
+};
+
+class FFabTedsMyFolderIntegration
+{
+public:
+	static void QueueSyncRequest();
+	static void QueueSyncRequest(uint32 BatchSize);
+
+private:
+	using HttpRequestPtr = TSharedPtr<class IHttpRequest, ESPMode::ThreadSafe>;
+	using HttpResponsePtr = TSharedPtr<class IHttpResponse, ESPMode::ThreadSafe>;
+
+	static void QueueSyncRequest(const FString& Cursor, uint32 BatchSize);
+	static void ProcessSyncResults(uint32 BatchSize, HttpRequestPtr Request, HttpResponsePtr Response, bool bWasSuccessful);
+
+	static void SetNameColumn(FNameColumn& Target, const FJsonObject& Object, FString& Temp);
+	static void SetFabObjectColumn(FFabObjectColumn& Target, const FJsonObject& Object, FString& Temp);
+	static void SetUrlColumn(FEditorDataStorageUrlColumn& Target, const FJsonObject& Object, FString& Temp);
+	static void AddDistributionMethod(
+		IEditorDataStorageProvider* Storage, UE::Editor::DataStorage::RowHandle Row, const FJsonObject& Object, FString& Temp);
+	static void AddImages(
+		IEditorDataStorageProvider* Storage, UE::Editor::DataStorage::RowHandle Row, const FJsonObject& Object, FString& Temp);
+};
\ No newline at end of file
diff --git a/Engine/Source/Runtime/TypedElementFramework/Public/Elements/Columns/TypedElementWebColumns.h b/Engine/Source/Runtime/TypedElementFramework/Public/Elements/Columns/TypedElementWebColumns.h
new file mode 100644
--- /dev/null
+++ b/Engine/Source/Runtime/TypedElementFramework/Public/Elements/Columns/TypedElementWebColumns.h
@@ -0,0 +1,48 @@
+// Copyright Epic Games, Inc. All Rights Reserved.
+
+#pragma once
+
+#include "Elements/Common/TypedElementCommonTypes.h"
+#include "UObject/ObjectMacros.h"
+
+#include "TypedElementWebColumns.generated.h"
+
+/**
+ * Information for retrieving a web page from the web.
+ */
+USTRUCT(meta = (DisplayName = "Url"))
+struct FEditorDataStorageUrlColumn final : public FEditorDataStorageColumn
+{
+	GENERATED_BODY()
+
+	UPROPERTY()
+	FString UrlString;
+};
+
+/**
+ * Information for retrieving an image from the web.
+ * The width and height can optionally be set to non-zero indicating that the size should overwrite the size from the source image.
+ */
+USTRUCT(meta = (DisplayName = "WebImage", EditorDataStorage_DynamicColumnTemplate))
+struct FEditorDataStorageWebImageColumn final : public FEditorDataStorageColumn
+{
+	GENERATED_BODY()
+
+	/** URL to the image in string form. Can be used with for instance with WebImage(Cache) to create a local image. */
+	UPROPERTY()
+	FString UrlString;
+
+	/** If set, the target width for the image. A value of zero indicates that the width from the source image should be used. */
+	UPROPERTY(meta = (IgnoreForMemberInitializationTest))
+	uint16 Width;
+
+	/** If set, the target height for the image. A value of zero indicates that the height from the source image should be used. */
+	UPROPERTY(meta = (IgnoreForMemberInitializationTest))
+	uint16 Height;
+};
+
+namespace UE::Editor::DataStorage
+{
+	using FUrlColumn = FEditorDataStorageUrlColumn;
+	using FWebImageColumn = FEditorDataStorageWebImageColumn;
+}
\ No newline at end of file
