Fixed build error.

#rnx

diff --git a/Engine/Plugins/Fab/Source/Fab/Private/Teds/FabMyFolderIntegration.cpp b/Engine/Plugins/Fab/Source/Fab/Private/Teds/FabMyFolderIntegration.cpp
--- a/Engine/Plugins/Fab/Source/Fab/Private/Teds/FabMyFolderIntegration.cpp
+++ b/Engine/Plugins/Fab/Source/Fab/Private/Teds/FabMyFolderIntegration.cpp
@@ -43,7 +43,7 @@ void FFabTedsMyFolderIntegration::QueueSyncRequest(const FString& Cursor, uint32
 {
 	if (EOS_Auth_GetLoggedInAccountsCount(FabAuthentication::AuthHandle) > 0)
 	{
-		const UFabSettings* FabPluginSettings = GetDefault<UFabSettings>();
+		const UFabSettings* Settings = GetDefault<UFabSettings>();
 
 		EOS_Auth_Token* UserAuthToken = nullptr;
 
@@ -66,7 +66,7 @@ void FFabTedsMyFolderIntegration::QueueSyncRequest(const FString& Cursor, uint32
 			HttpRequest->SetHeader("accept", "application/json");
 			HttpRequest->SetHeader("Authorization", FString::Printf(TEXT("Bearer %s"), *FString(UserAuthToken->AccessToken)));
 			HttpRequest->SetURL(FString::Printf(TEXT("%s/e/accounts/%s/ue/library?count=%u%s"),
-				*FabPluginSettings->GetUrlFromEnvironment(),
+				*Settings->GetUrlFromEnvironment(),
 				*LexToString(FabAuthentication::EpicAccountId), BatchSize, *CursorString));
 			HttpRequest->OnProcessRequestComplete().BindLambda([BatchSize](HttpRequestPtr Request, HttpResponsePtr Response, bool bWasSuccessful)
 				{
