Apply same platform filters to EOSShared as Fab so that EOSShared doesn't get included by default on all platforms.

diff --git a/Engine/Plugins/Fab/Fab.uplugin b/Engine/Plugins/Fab/Fab.uplugin
--- a/Engine/Plugins/Fab/Fab.uplugin
+++ b/Engine/Plugins/Fab/Fab.uplugin
@@ -33,7 +33,13 @@
 		},
 		{
 			"Name": "EOSShared",
-			"Enabled": true
+			"Enabled": true,
+			"PlatformAllowList" :
+			[
+				"Linux",
+				"Win64",
+				"Mac"
+			]
 		}
 	],
 	"Modules": [
