Fab: Add a new Fab Content Source to display the contents of "My Folder" in the Content Browser

Content Sources: Add an init param to allow content sources to request a default tile view widget

diff --git a/Engine/Plugins/Experimental/EditorDataStorage/Source/TedsUI/Private/Widgets/UrlWidget.cpp b/Engine/Plugins/Experimental/EditorDataStorage/Source/TedsUI/Private/Widgets/UrlWidget.cpp
new file mode 100644
--- /dev/null
+++ b/Engine/Plugins/Experimental/EditorDataStorage/Source/TedsUI/Private/Widgets/UrlWidget.cpp
@@ -0,0 +1,50 @@
+// Copyright Epic Games, Inc. All Rights Reserved.
+
+#include "Widgets/UrlWidget.h"
+
+#include "Elements/Columns/TypedElementWebColumns.h"
+#include "Elements/Framework/TypedElementAttributeBinding.h"
+#include "Internationalization/Text.h"
+#include "Widgets/Input/SHyperlink.h"
+#include "Widgets/SBoxPanel.h"
+
+void UUrlWidgetFactory::RegisterWidgetConstructors(IEditorDataStorageProvider& DataStorage,
+														IEditorDataStorageUiProvider& DataStorageUi) const
+{
+	using namespace UE::Editor::DataStorage::Queries;
+	DataStorageUi.RegisterWidgetFactory<FUrlWidgetConstructor>(TEXT("General.Cell"), TColumn<FUrlColumn>());
+}
+
+FUrlWidgetConstructor::FUrlWidgetConstructor()
+	: FSimpleWidgetConstructor(StaticStruct())
+{
+}
+
+
+
+TSharedPtr<SWidget> FUrlWidgetConstructor::CreateWidget(IEditorDataStorageProvider* DataStorage,
+                                                               IEditorDataStorageUiProvider* DataStorageUi,
+                                                               UE::Editor::DataStorage::RowHandle TargetRow,
+                                                               UE::Editor::DataStorage::RowHandle WidgetRow,
+                                                               const UE::Editor::DataStorage::FMetaDataView& Arguments)
+{
+	UE::Editor::DataStorage::FAttributeBinder Binder(TargetRow, DataStorage);
+	
+	return SNew(SHorizontalBox)
+			+SHorizontalBox::Slot()
+			.AutoWidth()
+			.HAlign(HAlign_Center)
+			.VAlign(VAlign_Center)
+			[
+				SNew(SHyperlink)
+					.Text(Binder.BindText(&UE::Editor::DataStorage::FUrlColumn::UrlString))
+					.Style(FAppStyle::Get(), "Common.GotoBlueprintHyperlink")
+					.OnNavigate_Lambda([DataStorage, TargetRow] ()
+					{
+						if (UE::Editor::DataStorage::FUrlColumn* UrlColumn = DataStorage->GetColumn<UE::Editor::DataStorage::FUrlColumn>(TargetRow))
+						{
+							FPlatformProcess::LaunchURL(*UrlColumn->UrlString, nullptr, nullptr);
+						}
+					})
+			];
+}
diff --git a/Engine/Plugins/Experimental/EditorDataStorage/Source/TedsUI/Private/Widgets/UrlWidget.h b/Engine/Plugins/Experimental/EditorDataStorage/Source/TedsUI/Private/Widgets/UrlWidget.h
new file mode 100644
--- /dev/null
+++ b/Engine/Plugins/Experimental/EditorDataStorage/Source/TedsUI/Private/Widgets/UrlWidget.h
@@ -0,0 +1,42 @@
+// Copyright Epic Games, Inc. All Rights Reserved.
+
+#pragma once
+
+#include "Elements/Interfaces/TypedElementDataStorageFactory.h"
+#include "Elements/Interfaces/TypedElementDataStorageUiInterface.h"
+
+#include "UrlWidget.generated.h"
+
+class IEditorDataStorageProvider;
+class IEditorDataStorageUiProvider;
+class UScriptStruct;
+
+UCLASS()
+class UUrlWidgetFactory : public UEditorDataStorageFactory
+{
+	GENERATED_BODY()
+
+public:
+	~UUrlWidgetFactory() override = default;
+
+	void RegisterWidgetConstructors(IEditorDataStorageProvider& DataStorage,
+		IEditorDataStorageUiProvider& DataStorageUi) const override;
+};
+
+// Widget to display a URL in TEDS
+USTRUCT()
+struct FUrlWidgetConstructor : public FSimpleWidgetConstructor
+{
+	GENERATED_BODY()
+
+public:
+	FUrlWidgetConstructor();
+	~FUrlWidgetConstructor() override = default;
+
+	virtual TSharedPtr<SWidget> CreateWidget(
+		IEditorDataStorageProvider* DataStorage,
+		IEditorDataStorageUiProvider* DataStorageUi,
+		UE::Editor::DataStorage::RowHandle TargetRow,
+		UE::Editor::DataStorage::RowHandle WidgetRow, 
+		const UE::Editor::DataStorage::FMetaDataView& Arguments) override;
+};
\ No newline at end of file
diff --git a/Engine/Plugins/Experimental/EditorDataStorageFeatures/Source/TedsContentBrowser/Private/Widgets/TedsContentBrowserAssetViewWidget.cpp b/Engine/Plugins/Experimental/EditorDataStorageFeatures/Source/TedsContentBrowser/Private/Widgets/TedsContentBrowserAssetViewWidget.cpp
--- a/Engine/Plugins/Experimental/EditorDataStorageFeatures/Source/TedsContentBrowser/Private/Widgets/TedsContentBrowserAssetViewWidget.cpp
+++ b/Engine/Plugins/Experimental/EditorDataStorageFeatures/Source/TedsContentBrowser/Private/Widgets/TedsContentBrowserAssetViewWidget.cpp
@@ -8,6 +8,7 @@
 #include "Elements/Columns/TypedElementLabelColumns.h"
 #include "Elements/Columns/TypedElementTypeInfoColumns.h"
 #include "Widgets/STedsTableViewer.h"
+#include "Widgets/STedsTileViewer.h"
 
 #define LOCTEXT_NAMESPACE "ContentBrowserAssetViewWidget"
 
@@ -34,13 +35,35 @@ public:
 		RowView = MakeShared<QueryStack::FRowViewNode>(FRowHandleArrayView(Rows.GetData(), Rows.Num(),
 			FRowHandleArrayView::EFlags::IsSorted | FRowHandleArrayView::EFlags::IsUnique));
 
-		ChildSlot
-		[
-			SNew(UE::Editor::DataStorage::STedsTableViewer)
+		TSharedPtr<SWidget> ChildWidget;
+
+		switch (Arguments._InitParams.TableViewMode)
+		{
+		case ETableViewMode::Type::List:
+			ChildWidget = SNew(UE::Editor::DataStorage::STedsTableViewer)
 					.QueryStack(RowView)
 					.CellWidgetPurposes(Arguments._InitParams.WidgetPurposes)
 					.Columns(Arguments._InitParams.Columns)
-					.ListSelectionMode(ESelectionMode::Multi)
+					.ListSelectionMode(ESelectionMode::Multi);
+			break;
+		case ETableViewMode::Type::Tile:
+			ChildWidget = SNew(UE::Editor::DataStorage::STedsTileViewer)
+					.QueryStack(RowView)
+					// The Tile View only accepts one widget purposes intentionally, all other table viewers will do the same with the updated
+					// purposes system so all TArray args will be removed
+					.WidgetPurpose(Arguments._InitParams.WidgetPurposes[0])
+					.Columns(Arguments._InitParams.Columns)
+					.SelectionMode(ESelectionMode::Multi);
+			break;
+		case ETableViewMode::Type::Tree:
+
+		default:
+			ChildWidget = SNullWidget::NullWidget;
+		}
+
+		ChildSlot
+		[
+			ChildWidget.ToSharedRef()
 		];
 	}
 
diff --git a/Engine/Plugins/Fab/Source/Fab/Private/FabBrowser.cpp b/Engine/Plugins/Fab/Source/Fab/Private/FabBrowser.cpp
--- a/Engine/Plugins/Fab/Source/Fab/Private/FabBrowser.cpp
+++ b/Engine/Plugins/Fab/Source/Fab/Private/FabBrowser.cpp
@@ -258,6 +258,11 @@ FString FFabBrowser::GetUrl()
 	}
 }
 
+const ISlateStyle& FFabBrowser::GetStyleSet()
+{
+	return *(SlateStyleSet.Get());
+}
+
 TSharedRef<SDockTab> FFabBrowser::OpenTab(const FSpawnTabArgs&)
 {
 	LogEvent(
diff --git a/Engine/Plugins/Fab/Source/Fab/Private/FabBrowser.h b/Engine/Plugins/Fab/Source/Fab/Private/FabBrowser.h
--- a/Engine/Plugins/Fab/Source/Fab/Private/FabBrowser.h
+++ b/Engine/Plugins/Fab/Source/Fab/Private/FabBrowser.h
@@ -67,4 +67,5 @@ public:
 	static void ShowSettings();
 	static void OpenURL(const FString& InURL = GetUrl());
 	static FString GetUrl();
+	static const ISlateStyle& GetStyleSet();
 };
diff --git a/Engine/Plugins/Fab/Source/Fab/Private/Teds/FabContentSource.cpp b/Engine/Plugins/Fab/Source/Fab/Private/Teds/FabContentSource.cpp
new file mode 100644
--- /dev/null
+++ b/Engine/Plugins/Fab/Source/Fab/Private/Teds/FabContentSource.cpp
@@ -0,0 +1,71 @@
+// Copyright Epic Games, Inc. All Rights Reserved.
+
+#include "FabContentSource.h"
+
+#include "FabBrowser.h"
+#include "FabMyFolderIntegration.h"
+#include "IContentBrowserSingleton.h"
+#include "Elements/Columns/TypedElementMiscColumns.h"
+#include "Elements/Columns/TypedElementWebColumns.h"
+#include "Elements/Framework/TypedElementQueryBuilder.h"
+
+#define LOCTEXT_NAMESPACE "FabContentSource"
+
+namespace FabContentSource::Private
+{
+	static FName ContentSourceName("FabMyLibraryContentSource");
+	
+	static bool bEnableContentSource = false;
+	
+	static FAutoConsoleVariableRef CVarEnableTestContentSource(
+		TEXT("Fab.TEDS.MyLibrary.ContentSource"),
+		bEnableContentSource,
+		TEXT("Add a Content Source that displays your fab library in the Content Browser ")
+		, FConsoleVariableDelegate::CreateLambda([](IConsoleVariable* /*CVar*/)
+		{
+			if(bEnableContentSource)
+			{
+				IContentBrowserSingleton::Get().RegisterContentSourceFactory(ContentSourceName,
+					IContentBrowserSingleton::FContentSourceFactory::CreateLambda(
+					[]() -> TSharedRef<UE::Editor::ContentBrowser::IContentSource>
+					{
+						return MakeShared<FFabMyLibraryContentSource>();
+					}));
+			}
+			else
+			{
+				IContentBrowserSingleton::Get().UnregisterContentSourceFactory(ContentSourceName);
+			}
+		}));
+}
+
+FName FFabMyLibraryContentSource::GetName()
+{
+	return FabContentSource::Private::ContentSourceName;
+}
+
+FText FFabMyLibraryContentSource::GetDisplayName()
+{
+	return LOCTEXT("MyLibraryContentSourceDisplayName", "Fab");
+}
+
+FSlateIcon FFabMyLibraryContentSource::GetIcon()
+{
+	return FSlateIcon(FFabBrowser::GetStyleSet().GetStyleSetName(), "Fab.ToolbarIcon");
+}
+
+void FFabMyLibraryContentSource::GetAssetViewInitParams(UE::Editor::ContentBrowser::FTableViewerInitParams& OutInitParams)
+{
+	using namespace UE::Editor::DataStorage::Queries;
+	
+	OutInitParams.QueryDescription = Select()
+			.Where()
+				.All<FFabObjectColumn>()
+			.Compile();
+
+	OutInitParams.Columns = { FNameColumn::StaticStruct(), FUrlColumn::StaticStruct() };
+
+	OutInitParams.WidgetPurposes = { FName("General.Cell") }; 
+}
+
+#undef LOCTEXT_NAMESPACE
\ No newline at end of file
diff --git a/Engine/Plugins/Fab/Source/Fab/Private/Teds/FabContentSource.h b/Engine/Plugins/Fab/Source/Fab/Private/Teds/FabContentSource.h
new file mode 100644
--- /dev/null
+++ b/Engine/Plugins/Fab/Source/Fab/Private/Teds/FabContentSource.h
@@ -0,0 +1,15 @@
+// Copyright Epic Games, Inc. All Rights Reserved.
+
+#pragma once
+#include "ContentSources/IContentSource.h"
+
+class FFabMyLibraryContentSource : public UE::Editor::ContentBrowser::IContentSource
+{
+public:
+	virtual ~FFabMyLibraryContentSource() override = default;
+
+	virtual FName GetName() override;
+	virtual FText GetDisplayName() override;
+	virtual FSlateIcon GetIcon() override;
+	virtual void GetAssetViewInitParams(UE::Editor::ContentBrowser::FTableViewerInitParams& OutInitParams) override;
+};
\ No newline at end of file
diff --git a/Engine/Source/Editor/ContentBrowser/Private/ContentBrowserStyle.cpp b/Engine/Source/Editor/ContentBrowser/Private/ContentBrowserStyle.cpp
--- a/Engine/Source/Editor/ContentBrowser/Private/ContentBrowserStyle.cpp
+++ b/Engine/Source/Editor/ContentBrowser/Private/ContentBrowserStyle.cpp
@@ -396,9 +396,11 @@ namespace UE::ContentBrowser::Private
 			
 			ContentBrowserSourceBarStyle.SetIconSize(DefaultIconSize);
 			ContentBrowserSourceBarStyle.SetUniformBlockWidth(DefaultButtonSize);
-			ContentBrowserSourceBarStyle.SetButtonPadding(FMargin(2.0f));
-			ContentBrowserSourceBarStyle.ButtonStyle.SetNormalPadding(FMargin(0));
-			ContentBrowserSourceBarStyle.ButtonStyle.SetPressedPadding(FMargin(0));
+			ContentBrowserSourceBarStyle.SetUniformBlockHeight(DefaultButtonSize);
+			ContentBrowserSourceBarStyle.SetButtonPadding(FMargin(0.0f, 2.0f));
+			ContentBrowserSourceBarStyle.ButtonStyle.SetNormalPadding(FMargin(0.0f));
+			ContentBrowserSourceBarStyle.ButtonStyle.SetPressedPadding(FMargin(0.0f));
+			ContentBrowserSourceBarStyle.SetBackgroundPadding(FMargin(4.0f));
 			
 			Set("ContentBrowser.SourceBar", ContentBrowserSourceBarStyle);
 			Set("ContentBrowser.Sources.PanelIcon.Opened", new CORE_IMAGE_BRUSH_SVG("Starship/Common/SidePanelLeft", DefaultIconSize));
diff --git a/Engine/Source/Editor/ContentBrowser/Private/ContentSources/Widgets/SContentSource.cpp b/Engine/Source/Editor/ContentBrowser/Private/ContentSources/Widgets/SContentSource.cpp
--- a/Engine/Source/Editor/ContentBrowser/Private/ContentSources/Widgets/SContentSource.cpp
+++ b/Engine/Source/Editor/ContentBrowser/Private/ContentSources/Widgets/SContentSource.cpp
@@ -29,8 +29,20 @@ namespace UE::Editor::ContentBrowser
 
 	void SContentSource::SetContentSource(const TSharedPtr<IContentSource>& InContentSource)
 	{
+		// Let the old content source know it is being disabled
+		if (ContentSource)
+		{
+			ContentSource->OnContentSourceDisabled();
+		}
+		
 		ContentSource = InContentSource;
 
+		// Let the new Content Source know it has been enabled
+		if (ContentSource)
+		{
+			ContentSource->OnContentSourceEnabled();
+		}
+
 		ChildSlot
 		[
 			CreateWidget()
diff --git a/Engine/Source/Editor/ContentBrowser/Public/ContentSources/IContentSource.h b/Engine/Source/Editor/ContentBrowser/Public/ContentSources/IContentSource.h
--- a/Engine/Source/Editor/ContentBrowser/Public/ContentSources/IContentSource.h
+++ b/Engine/Source/Editor/ContentBrowser/Public/ContentSources/IContentSource.h
@@ -3,6 +3,7 @@
 #pragma once
 
 #include "Elements/Common/TypedElementQueryDescription.h"
+#include "Framework/Views/ITypedTableView.h"
 #include "Internationalization/Text.h"
 #include "Textures/SlateIcon.h"
 #include "UObject/NameTypes.h"
@@ -22,6 +23,9 @@ namespace UE::Editor::ContentBrowser
 
 		// The widget purposes used by the table viewer to display the widgets in the cells
 		TArray<FName> WidgetPurposes;
+
+		// The table view type to use, currently only supports list and tile
+		ETableViewMode::Type TableViewMode = ETableViewMode::Type::List;
 	};
 
 	/*
@@ -50,5 +54,11 @@ namespace UE::Editor::ContentBrowser
 
 		// Get the init params used by the TEDS Table Viewer when this content source is active
 		virtual void GetAssetViewInitParams(FTableViewerInitParams& OutInitParams) = 0;
+
+		// Called when this content source is switched in
+		virtual void OnContentSourceEnabled() {}
+
+		// Called when this content source is swapped out
+		virtual void OnContentSourceDisabled() {}
 	};
 }
\ No newline at end of file
