[Fab] Update plugin version to 0.0.5

#virtualized

diff --git a/Engine/Plugins/Fab/Fab.uplugin b/Engine/Plugins/Fab/Fab.uplugin
--- a/Engine/Plugins/Fab/Fab.uplugin
+++ b/Engine/Plugins/Fab/Fab.uplugin
@@ -1,7 +1,7 @@
 {
 	"FileVersion": 3,
-	"Version": 3,
-	"VersionName": "0.0.3",
+	"Version": 5,
+	"VersionName": "0.0.5",
 	"FriendlyName": "Fab",
 	"Description": "Fab Plugin",
 	"Category": "Other",
diff --git a/Engine/Plugins/Fab/Source/Fab/Private/FabAuthentication.cpp b/Engine/Plugins/Fab/Source/Fab/Private/FabAuthentication.cpp
--- a/Engine/Plugins/Fab/Source/Fab/Private/FabAuthentication.cpp
+++ b/Engine/Plugins/Fab/Source/Fab/Private/FabAuthentication.cpp
@@ -96,11 +96,6 @@ namespace FabAuthentication
 			}
 
 			PlatformHandle = SDKManager->CreatePlatform(PlatformOptions);
-
-			// Sequence of events
-			// 1. Login using persist
-			// 2. If that fails - login using the exchange code
-			LoginUsingPersist();
 		}
 	}
 
@@ -271,7 +266,6 @@ namespace FabAuthentication
 
 		EOS_Auth_LoginOptions LoginOptions = {0};
 		LoginOptions.ApiVersion = EOS_AUTH_LOGIN_API_LATEST;
-		LoginOptions.ScopeFlags = EOS_EAuthScopeFlags::EOS_AS_BasicProfile;
 		LoginOptions.Credentials = &Credentials;
 
 		EOS_Auth_Login(AuthHandle, &LoginOptions, nullptr, ExchangeCodeLoginCompleteCallbackFn);
@@ -292,7 +286,6 @@ namespace FabAuthentication
 
 		EOS_Auth_LoginOptions LoginOptions = {0};
 		LoginOptions.ApiVersion = EOS_AUTH_LOGIN_API_LATEST;
-		LoginOptions.ScopeFlags = EOS_EAuthScopeFlags::EOS_AS_BasicProfile;
 		LoginOptions.Credentials = &Credentials;
 
 		EOS_Auth_Login(AuthHandle, &LoginOptions, nullptr, PersistLoginCompleteCallbackFn);
@@ -313,7 +306,6 @@ namespace FabAuthentication
 
 		EOS_Auth_LoginOptions LoginOptions = {0};
 		LoginOptions.ApiVersion = EOS_AUTH_LOGIN_API_LATEST;
-		LoginOptions.ScopeFlags = EOS_EAuthScopeFlags::EOS_AS_BasicProfile;
 		LoginOptions.Credentials = &Credentials;
 
 		EOS_Auth_Login(AuthHandle, &LoginOptions, nullptr, AccountPortalLoginCompleteCallbackFn);
diff --git a/Engine/Plugins/Fab/Source/Fab/Private/FabBrowser.cpp b/Engine/Plugins/Fab/Source/Fab/Private/FabBrowser.cpp
--- a/Engine/Plugins/Fab/Source/Fab/Private/FabBrowser.cpp
+++ b/Engine/Plugins/Fab/Source/Fab/Private/FabBrowser.cpp
@@ -3,7 +3,7 @@
 #include "FabBrowser.h"
 
 #include "ContentBrowserModule.h"
-#include "FabBrowserApi.h"
+#include "FabAuthentication.h"
 #include "FabLog.h"
 #include "FabSettings.h"
 #include "FabSettingsWindow.h"
@@ -87,7 +87,7 @@ void FFabBrowser::RegisterSlateStyle()
 	SlateStyleSet = MakeUnique<FSlateStyleSet>(TEXT("FabStyle"));
 	SlateStyleSet->SetContentRoot(IPluginManager::Get().FindPlugin(TEXT("Fab"))->GetBaseDir() / TEXT("Resources"));
 
-	const FString IconPath = SlateStyleSet->RootToContentDir(TEXT("FabLogo.svg"));
+	const FString IconPath          = SlateStyleSet->RootToContentDir(TEXT("FabLogo.svg"));
 	const FString AlternateIconPath = SlateStyleSet->RootToContentDir(TEXT("FabLogoAlternate.svg"));
 	SlateStyleSet->Set(FabMenuIconName, new FSlateVectorImageBrush(IconPath, CoreStyleConstants::Icon16x16));
 	SlateStyleSet->Set(FabAssetIconName, new FSlateVectorImageBrush(AlternateIconPath, CoreStyleConstants::Icon20x20));
@@ -265,14 +265,8 @@ const ISlateStyle& FFabBrowser::GetStyleSet()
 
 TSharedRef<SDockTab> FFabBrowser::OpenTab(const FSpawnTabArgs&)
 {
-	LogEvent(
-		{
-			"Open Tab",
-			"Plugin",
-			"Click"
-		}
-	);
-
+	FabAuthentication::LoginUsingPersist();
+	
 	FabPluginSettings = GetDefault<UFabSettings>();
 
 	JavascriptApi = NewObject<UFabBrowserApi>();
@@ -287,9 +281,9 @@ TSharedRef<SDockTab> FFabBrowser::OpenTab(const FSpawnTabArgs&)
 
 	FCreateBrowserWindowSettings WindowSettings;
 
-	FString PluginPath = IPluginManager::Get().FindPlugin(TEXT("Fab"))->GetBaseDir();
-	FString IndexUrl   = FPaths::ConvertRelativePathToFull(FPaths::Combine(PluginPath, TEXT("ThirdParty"), TEXT("index.html")));
-	FString FinalUrl   = FPaths::Combine(TEXT("file:///"), IndexUrl);
+	FString PluginPath              = IPluginManager::Get().FindPlugin(TEXT("Fab"))->GetBaseDir();
+	FString IndexUrl                = FPaths::ConvertRelativePathToFull(FPaths::Combine(PluginPath, TEXT("ThirdParty"), TEXT("index.html")));
+	FString FinalUrl                = FPaths::Combine(TEXT("file:///"), IndexUrl);
 	WindowSettings.InitialURL       = FinalUrl;
 	WindowSettings.BrowserFrameRate = 60;
 
@@ -299,6 +293,34 @@ TSharedRef<SDockTab> FFabBrowser::OpenTab(const FSpawnTabArgs&)
 	WebBrowserWindow = WebBrowserSingleton->CreateBrowserWindow(WindowSettings);
 	WebBrowserWindow->OnUnhandledKeyUp().BindLambda([](const FKeyEvent&) { return true; });
 	WebBrowserWindow->OnUnhandledKeyDown().BindLambda([](const FKeyEvent&) { return true; });
+	WebBrowserWindow->OnUrlChanged().AddLambda(
+		[](const FString& Url)
+		{
+			if (FabPluginSettings->Environment != EFabEnvironment::Prod)
+				return;
+
+			FString Domain, Protocol;
+			Url.Split("://", &Protocol, &Domain);
+			if (!Protocol.Contains("http"))
+			{
+				return;
+			}
+
+			if (int32 PathIndex; Domain.FindChar('/', PathIndex))
+			{
+				Domain = Domain.Left(PathIndex);
+			}
+
+			Domain = Domain.Replace(TEXT("www."), TEXT(""));
+
+			if (!Domain.Contains("fab.com"))
+			{
+				FAB_LOG_ERROR("Trying to access thirdparty url [%s] in plugin browser. Redirecting back to fab.com", *Url);
+				WebBrowserWindow->LoadURL(GetUrl());
+				FPlatformProcess::LaunchURL(*Url, nullptr, nullptr);
+			}
+		}
+	);
 
 	if (FabPluginSettings->bEnableDebugOptions)
 	{
@@ -329,15 +351,7 @@ TSharedRef<SDockTab> FFabBrowser::OpenTab(const FSpawnTabArgs&)
 	WebBrowserInstance->BindUObject(TEXT("fab"), Cast<UObject>(JavascriptApi), true);
 	WebBrowserWindow->Reload();
 
-	SAssignNew(DockTab, SDockTab).TabRole(ETabRole::NomadTab).OnTabClosed_Lambda(
-		[](TSharedRef<class SDockTab> InParentTab)
-		{
-			WebBrowserInstance->UnbindUObject(TEXT("fab"), Cast<UObject>(JavascriptApi), true);
-			WebBrowserInstance.Reset();
-			WebBrowserWindow.Reset();
-			DockTab.Reset();
-		}
-	)
+	SAssignNew(DockTab, SDockTab).TabRole(ETabRole::NomadTab).OnTabClosed_Static(&FFabBrowser::OnPluginTabClosed)
 	[
 		WebBrowserInstance.ToSharedRef()
 	];
@@ -345,6 +359,37 @@ TSharedRef<SDockTab> FFabBrowser::OpenTab(const FSpawnTabArgs&)
 	return DockTab.ToSharedRef();
 }
 
+void FFabBrowser::OnPluginTabClosed(TSharedRef<SDockTab> InParentTab)
+{
+	LogEvent(
+		{
+			"click",
+			"button",
+			"terminatePlugin",
+			"closeFabPlugin",
+			"interaction",
+			{
+				"Fab_UE5_Plugin",
+				UFabBrowserApi::GetApiVersion()
+			}
+		}
+	);
+
+	AsyncTask(ENamedThreads::AnyBackgroundThreadNormalTask, []()
+	{
+		// add a small delay so that the frontend can register the close event
+		FPlatformProcess::Sleep(1.5f);
+
+		AsyncTask(ENamedThreads::GameThread, []()
+		{
+			WebBrowserInstance->UnbindUObject(TEXT("fab"), Cast<UObject>(JavascriptApi), true);
+			WebBrowserInstance.Reset();
+			WebBrowserWindow.Reset();
+			DockTab.Reset();
+		});
+	});
+}
+
 void FFabBrowser::ExecuteJavascript(const FString& InSrcScript)
 {
 	if (WebBrowserInstance == nullptr)
@@ -377,10 +422,9 @@ void FFabBrowser::GetSignedUrl(const FString& AssetId, const int32 Tier)
 void FFabBrowser::LogEvent(const FFabAnalyticsPayload& Payload)
 {
 	FString JSONPayload;
-	FJsonObjectConverter::UStructToJsonObjectString(Payload, JSONPayload, 0, 0);
-	FAB_LOG("%s", *JSONPayload);
+	FJsonObjectConverter::UStructToJsonObjectString(Payload, JSONPayload, 0, 0, 0, nullptr, false);
 
-	// ExecuteJavascript(FString::Printf(TEXT("window.ue.fab.onDownloaded()")));
+	ExecuteJavascript(FString::Printf(TEXT("window.ue.fab.logevent('%s')"), *JSONPayload));
 }
 
 void FFabBrowser::ShowSettings()
diff --git a/Engine/Plugins/Fab/Source/Fab/Private/FabBrowser.h b/Engine/Plugins/Fab/Source/Fab/Private/FabBrowser.h
--- a/Engine/Plugins/Fab/Source/Fab/Private/FabBrowser.h
+++ b/Engine/Plugins/Fab/Source/Fab/Private/FabBrowser.h
@@ -4,6 +4,7 @@
 
 #include "CoreMinimal.h"
 #include "ContentBrowserDelegates.h"
+#include "FabBrowserApi.h"
 #include "FabBrowser.generated.h"
 
 class SDockTab;
@@ -15,6 +16,18 @@ class UFabSettings;
 class FSlateStyleSet;
 class FExtender;
 
+USTRUCT()
+struct FFabAnalyticsEventValue
+{
+	GENERATED_BODY()
+
+	UPROPERTY()
+	FString Platform;
+	
+	UPROPERTY()
+	FFabApiVersion ApiVersion;
+};
+
 USTRUCT()
 struct FFabAnalyticsPayload
 {
@@ -28,6 +41,15 @@ struct FFabAnalyticsPayload
 
 	UPROPERTY()
 	FString EventAction;
+
+	UPROPERTY()
+	FString EventLabel;
+
+	UPROPERTY()
+	FString EventType;
+
+	UPROPERTY()
+	FFabAnalyticsEventValue EventValue;
 };
 
 class FFabBrowser
@@ -53,6 +75,7 @@ private:
 	static void SetupEntryPoints();
 	static void ExecuteJavascript(const FString& InSrcScript);
 	static TSharedRef<SDockTab> OpenTab(const FSpawnTabArgs& InArgs);
+	static void OnPluginTabClosed(TSharedRef<class SDockTab> InParentTab);
 	static TSharedRef<FExtender> OnExtendContentBrowserAssetSelectionMenu(const TArray<FAssetData>& SelectedAssets);
 	static TSharedRef<SWidget> OnFabAssetIconGenerate(const FAssetData& AssetData);
 
diff --git a/Engine/Plugins/Fab/Source/Fab/Private/FabBrowserApi.cpp b/Engine/Plugins/Fab/Source/Fab/Private/FabBrowserApi.cpp
--- a/Engine/Plugins/Fab/Source/Fab/Private/FabBrowserApi.cpp
+++ b/Engine/Plugins/Fab/Source/Fab/Private/FabBrowserApi.cpp
@@ -5,11 +5,13 @@
 #include "FabBrowser.h"
 #include "FabSettings.h"
 #include "FabLog.h"
+#include "HAL/PlatformApplicationMisc.h"
 #include "Interfaces/IPluginManager.h"
+#include "Kismet/GameplayStatics.h"
+#include "Misc/EngineVersion.h"
 #include "Workflows/GenericImportWorkflow.h"
 #include "Workflows/PackImportWorkflow.h"
 #include "Workflows/QuixelImportWorkflow.h"
-#include "Misc/EngineVersion.h"
 #include "Workflows/GenericDragDropWorkflow.h"
 #include "Workflows/QuixelDragDropWorkflow.h"
 
@@ -200,7 +202,6 @@ void UFabBrowserApi::Logout()
 
 void UFabBrowserApi::OpenPluginSettings()
 {
-	FAB_LOG("Open plugin settings");
 	FFabBrowser::ShowSettings();
 }
 
@@ -278,6 +279,8 @@ FFabApiVersion UFabBrowserApi::GetApiVersion()
 		ApiVersion.PluginVersion = PluginVersion;
 	}
 
+	ApiVersion.Platform = UGameplayStatics::GetPlatformName();
+
 	return ApiVersion;
 }
 
@@ -286,6 +289,28 @@ void UFabBrowserApi::OpenUrlInBrowser(const FString& Url)
 	FPlatformProcess::LaunchURL(*Url, nullptr, nullptr);
 }
 
+void UFabBrowserApi::CopyToClipboard(const FString& Content)
+{
+	FPlatformApplicationMisc::ClipboardCopy(*Content);
+}
+
+void UFabBrowserApi::PluginOpened()
+{
+	FFabBrowser::LogEvent(
+		{
+			"click",
+			"button",
+			"startPlugin",
+			"openFabPlugin",
+			"interaction",
+			{
+				"Fab_UE5_Plugin",
+				GetApiVersion()
+			}
+		}
+	);
+}
+
 FString UFabBrowserApi::GetUrl()
 {
 	return FFabBrowser::GetUrl();
diff --git a/Engine/Plugins/Fab/Source/Fab/Private/FabBrowserApi.h b/Engine/Plugins/Fab/Source/Fab/Private/FabBrowserApi.h
--- a/Engine/Plugins/Fab/Source/Fab/Private/FabBrowserApi.h
+++ b/Engine/Plugins/Fab/Source/Fab/Private/FabBrowserApi.h
@@ -45,9 +45,12 @@ struct FFabApiVersion
 
 	UPROPERTY()
 	FString Api;
-	
+
 	UPROPERTY()
 	FString PluginVersion;
+	
+	UPROPERTY()
+	FString Platform;
 };
 
 USTRUCT()
@@ -57,7 +60,7 @@ struct FFabFrontendSettings
 
 	UPROPERTY()
 	FString PreferredFormat;
-	
+
 	UPROPERTY()
 	FString PreferredQuality;
 };
@@ -99,7 +102,7 @@ public:
 
 	UFUNCTION()
 	FString GetRefreshToken();
-	
+
 	UFUNCTION()
 	void OpenPluginSettings();
 
@@ -108,13 +111,19 @@ public:
 
 	UFUNCTION()
 	void SetPreferredQualityTier(const FString& PreferredQuality);
-	
+
 	UFUNCTION()
-	FFabApiVersion GetApiVersion();
+	static FFabApiVersion GetApiVersion();
 
 	UFUNCTION()
 	void OpenUrlInBrowser(const FString& Url);
 
+	UFUNCTION()
+	void CopyToClipboard(const FString& Content);
+
+	UFUNCTION()
+	void PluginOpened();
+	
 	UFUNCTION()
 	FString GetUrl();
 
diff --git a/Engine/Plugins/Fab/Source/Fab/Private/FabSettings.cpp b/Engine/Plugins/Fab/Source/Fab/Private/FabSettings.cpp
--- a/Engine/Plugins/Fab/Source/Fab/Private/FabSettings.cpp
+++ b/Engine/Plugins/Fab/Source/Fab/Private/FabSettings.cpp
@@ -1,9 +1,10 @@
 // Copyright Epic Games, Inc. All Rights Reserved.
 
 #include "FabSettings.h"
+
 #include "FabAuthentication.h"
 #include "FabBrowser.h"
-
+#include "Misc/Paths.h"
 #include "UObject/UnrealType.h"
 
 UFabSettings::UFabSettings()
@@ -25,7 +26,16 @@ void UFabSettings::PostEditChangeProperty(struct FPropertyChangedEvent& Property
 	Super::PostEditChangeProperty(PropertyChangedEvent);
 	if (PropertyChangedEvent.ChangeType != EPropertyChangeType::Interactive)
 	{
+		if (PropertyChangedEvent.GetMemberPropertyName() == GET_MEMBER_NAME_CHECKED(UFabSettings, CacheDirectoryPath))
+		{
+			if (!FPaths::DirectoryExists(CacheDirectoryPath.Path) || FPaths::IsRelative(CacheDirectoryPath.Path))
+			{
+				CacheDirectoryPath = FDirectoryPath { FPlatformProcess::UserTempDir() / FString("FabLibrary") };
+			}
+		}
+		
 		SaveConfig();
+		
 		if (PropertyChangedEvent.GetMemberPropertyName() == GET_MEMBER_NAME_CHECKED(UFabSettings, Environment))
 		{
 			FabAuthentication::DeletePersistentAuth();
diff --git a/Engine/Plugins/Fab/Source/Fab/Private/FabSettings.h b/Engine/Plugins/Fab/Source/Fab/Private/FabSettings.h
--- a/Engine/Plugins/Fab/Source/Fab/Private/FabSettings.h
+++ b/Engine/Plugins/Fab/Source/Fab/Private/FabSettings.h
@@ -2,6 +2,7 @@
 
 #pragma once
 
+#include "HAL/PlatformProcess.h"
 #include "Misc/Paths.h"
 #include "UObject/SoftObjectPath.h"
 
@@ -56,10 +57,14 @@ public:
 	UPROPERTY(config, EditAnywhere, Category = General)
 	bool bEnableDebugOptions = false;
 
+	/** Path to the local library */
+	UPROPERTY(config, EditAnywhere, Category = General)
+	FDirectoryPath CacheDirectoryPath { FPlatformProcess::UserTempDir() / FString("FabLibrary") };
+	
 	/** Cache directory */
 	UPROPERTY(config, VisibleAnywhere, Category = General)
 	FString CacheDirectorySize;
-
+	
 	/** Preferred default format */
 	/* The preferred format will always be selected, if not available, the best available format for the product will be chosen. */
 	UPROPERTY(config, VisibleAnywhere, Category = ProductFormats, meta=(DevOnly=true))
diff --git a/Engine/Plugins/Fab/Source/Fab/Private/Importers/QuixelGLTFImporter.cpp b/Engine/Plugins/Fab/Source/Fab/Private/Importers/QuixelGLTFImporter.cpp
--- a/Engine/Plugins/Fab/Source/Fab/Private/Importers/QuixelGLTFImporter.cpp
+++ b/Engine/Plugins/Fab/Source/Fab/Private/Importers/QuixelGLTFImporter.cpp
@@ -3,11 +3,11 @@
 #include "QuixelGLTFImporter.h"
 
 #include "Editor.h"
-#include "FoliageType_InstancedStaticMesh.h"
 #include "InterchangeGenericAssetsPipeline.h"
 #include "InterchangeGenericMaterialPipeline.h"
 #include "InterchangeGenericMeshPipeline.h"
 #include "InterchangeManager.h"
+#include "Runtime/Launch/Resources/Version.h"
 
 #if (ENGINE_MAJOR_VERSION >=5 && ENGINE_MINOR_VERSION <=3)
 #include "InterchangeglTFPipeline.h"
@@ -44,7 +44,7 @@ void FQuixelGltfImporter::SetupGlobalFoliageActor(const FString& ImportPath)
 	const FString GlobalFoliageActorPackageName     = "BP_GlobalFoliageActor";
 	const FString GlobalFoliageActorDestinationPath = FPaths::GetPath(FPaths::GetPath(ImportPath)) / GlobalFoliageActorPackageName;
 
-	if (FindPackage(nullptr, *GlobalFoliageActorDestinationPath))
+	if (IAssetRegistry::Get()->DoesPackageExistOnDisk(*GlobalFoliageActorDestinationPath) || FindPackage(nullptr, *GlobalFoliageActorDestinationPath))
 	{
 		return;
 	}
@@ -193,9 +193,10 @@ void FQuixelGltfImporter::ImportGltfImperfectionAsset(const FString& SourcePath,
 	UInterchangeGenericAssetsPipeline* AssetPipeline = GetGenericAssetPipeline(GeneratedPipelines);
 	if (AssetPipeline)
 	{
-		AssetPipeline->MeshPipeline->bImportStaticMeshes   = false;
-		AssetPipeline->MeshPipeline->bImportSkeletalMeshes = false;
-		AssetPipeline->MaterialPipeline->MaterialImport    = EInterchangeMaterialImportOption::ImportAsMaterialInstances;
+		AssetPipeline->MeshPipeline->bImportStaticMeshes                  = false;
+		AssetPipeline->MeshPipeline->bImportSkeletalMeshes                = false;
+		AssetPipeline->MaterialPipeline->bImportMaterials                 = false;
+		AssetPipeline->MaterialPipeline->TexturePipeline->bImportTextures = true;
 	}
 
 	UInterchangeMegascansPipeline* MegascanPipeline = GetMegascanPipeline(GeneratedPipelines);
@@ -293,11 +294,18 @@ void FQuixelGltfImporter::ImportGltfPlantAsset(const FString& SourcePath, const
 	{
 		AssetPipeline->MeshPipeline->bImportStaticMeshes                             = true;
 		AssetPipeline->MeshPipeline->bImportSkeletalMeshes                           = false;
+		//TODO: Move buildNanite into import pipeline
 		AssetPipeline->MeshPipeline->bBuildNanite                                    = bBuildNanite;
 		AssetPipeline->CommonMeshesProperties->bRecomputeNormals                     = true;
 		AssetPipeline->CommonMeshesProperties->bComputeWeightedNormals               = true;
 		AssetPipeline->MaterialPipeline->MaterialImport                              = EInterchangeMaterialImportOption::ImportAsMaterialInstances;
 		AssetPipeline->MaterialPipeline->TexturePipeline->bFlipNormalMapGreenChannel = true;
+
+		#if (ENGINE_MAJOR_VERSION >=5 && ENGINE_MINOR_VERSION >= 5)
+		AssetPipeline->MeshPipeline->bCollision = false;
+		#else
+		AssetPipeline->MeshPipeline->bImportCollision = false;
+		#endif
 	}
 
 	UInterchangeMegascansPipeline* MegascanPipeline = GetMegascanPipeline(GeneratedPipelines);
@@ -341,7 +349,7 @@ void FQuixelGltfImporter::ImportGltfPlantAsset(const FString& SourcePath, const
 	);
 }
 
-void FQuixelGltfImporter::ImportGltf3DAsset(const FString& SourcePath, const FString& DestinationPath, const bool bBuildNanite, TFunction<void(const TArray<UObject*>&)> OnDone)
+void FQuixelGltfImporter::ImportGltf3DAsset(const FString& SourcePath, const FString& DestinationPath, TFunction<void(const TArray<UObject*>&)> OnDone)
 {
 	const UInterchangeSourceData* InSourceData = UInterchangeManager::CreateSourceData(SourcePath);
 
@@ -352,7 +360,7 @@ void FQuixelGltfImporter::ImportGltf3DAsset(const FString& SourcePath, const FSt
 	{
 		AssetPipeline->MeshPipeline->bImportStaticMeshes   = true;
 		AssetPipeline->MeshPipeline->bImportSkeletalMeshes = false;
-		AssetPipeline->MeshPipeline->bBuildNanite          = bBuildNanite;
+		AssetPipeline->MeshPipeline->bBuildNanite          = true;
 		AssetPipeline->MaterialPipeline->MaterialImport    = EInterchangeMaterialImportOption::ImportAsMaterialInstances;
 	}
 
diff --git a/Engine/Plugins/Fab/Source/Fab/Private/Importers/QuixelGLTFImporter.h b/Engine/Plugins/Fab/Source/Fab/Private/Importers/QuixelGLTFImporter.h
--- a/Engine/Plugins/Fab/Source/Fab/Private/Importers/QuixelGLTFImporter.h
+++ b/Engine/Plugins/Fab/Source/Fab/Private/Importers/QuixelGLTFImporter.h
@@ -20,5 +20,5 @@ public:
 	static void ImportGltfImperfectionAsset(const FString& SourcePath, const FString& DestinationPath, TFunction<void(const TArray<UObject*>&)> OnDone);
 	static void ImportGltfSurfaceAsset(const FString& SourcePath, const FString& DestinationPath, TFunction<void(const TArray<UObject*>&)> OnDone);
 	static void ImportGltfPlantAsset(const FString& SourcePath, const FString& DestinationPath, const bool bBuildNanite, TFunction<void(const TArray<UObject*>&)> OnDone);
-	static void ImportGltf3DAsset(const FString& SourcePath, const FString& DestinationPath, const bool bBuildNanite, TFunction<void(const TArray<UObject*>&)> OnDone);
+	static void ImportGltf3DAsset(const FString& SourcePath, const FString& DestinationPath, TFunction<void(const TArray<UObject*>&)> OnDone);
 };
diff --git a/Engine/Plugins/Fab/Source/Fab/Private/Pipelines/InterchangeMegascansPipeline.cpp b/Engine/Plugins/Fab/Source/Fab/Private/Pipelines/InterchangeMegascansPipeline.cpp
--- a/Engine/Plugins/Fab/Source/Fab/Private/Pipelines/InterchangeMegascansPipeline.cpp
+++ b/Engine/Plugins/Fab/Source/Fab/Private/Pipelines/InterchangeMegascansPipeline.cpp
@@ -29,9 +29,17 @@
 #define MEGASCAN_BASE_KEY									TEXT("Megascan")
 
 #define MEGASCAN_MATERIAL_KEY							    MEGASCAN_BASE_KEY TEXT(".Material")
+
 #define MEGASCAN_MATERIAL_TYPE_KEY							MEGASCAN_MATERIAL_KEY TEXT(".Type")
+
 #define MEGASCAN_MATERIAL_BLEND_MODE_KEY				    MEGASCAN_MATERIAL_KEY TEXT(".BlendMode")
-#define MEGASCAN_MATERIAL_BLEND_MODE_OVERRIDE_KEY		    MEGASCAN_MATERIAL_KEY TEXT(".BlendModeOverride")
+#define MEGASCAN_MATERIAL_BLEND_MODE_VALUE_KEY				MEGASCAN_MATERIAL_BLEND_MODE_KEY TEXT(".Value")
+#define MEGASCAN_MATERIAL_BLEND_MODE_OVERRIDE_KEY		    MEGASCAN_MATERIAL_BLEND_MODE_KEY TEXT(".Override")
+
+#define MEGASCAN_MATERIAL_DISPLACEMENT_KEY					MEGASCAN_MATERIAL_KEY TEXT(".Displacement")
+#define MEGASCAN_MATERIAL_DISPLACEMENT_OVERRIDE_KEY			MEGASCAN_MATERIAL_DISPLACEMENT_KEY TEXT(".Override")
+#define MEGASCAN_MATERIAL_DISPLACEMENT_MAGNITUDE_KEY		MEGASCAN_MATERIAL_DISPLACEMENT_KEY TEXT(".Magnitude")
+#define MEGASCAN_MATERIAL_DISPLACEMENT_CENTER_KEY			MEGASCAN_MATERIAL_DISPLACEMENT_KEY TEXT(".Center")
 
 #define MEGASCAN_MESH_KEY									MEGASCAN_BASE_KEY TEXT(".Mesh")
 #define MEGASCAN_MESH_GENERATE_DISTANCE_FIELD_KEY			MEGASCAN_MESH_KEY TEXT(".GenerateDistanceField")
@@ -165,10 +173,18 @@ void UInterchangeMegascansPipeline::ExecutePostFactoryPipeline(
 	{
 		if (UStaticMesh* ImportedMesh = Cast<UStaticMesh>(CreatedAsset))
 		{
-			for (int32 Index = 0, MaxIndex = ImportedMesh->GetNumSourceModels(); Index < MaxIndex; ++Index)
+			const float BillboardScreenSizes[] = {
+				0.03f,
+				0.05f,
+				0.10f
+			};
+			int32 Index = 0;
+			for (const int32 MaxIndex = ImportedMesh->GetNumSourceModels(); Index < MaxIndex; ++Index)
 			{
-				ImportedMesh->GetSourceModel(Index).ScreenSize = (Index < MaxIndex - 1) ? FMath::Pow(0.75f, Index) : (0.05 * static_cast<int8>(MegascanAssetTier));
+				ImportedMesh->GetSourceModel(Index).ScreenSize = FMath::Pow(0.75f, Index);
 			}
+			if (const int ScreenSizeIndex = static_cast<int8>(MegascanAssetTier) - 1; ScreenSizeIndex >= 0)
+				ImportedMesh->GetSourceModel(Index - 1).ScreenSize = BillboardScreenSizes[ScreenSizeIndex];
 		}
 	}
 
@@ -409,8 +425,8 @@ void UInterchangeMegascansPipeline::SetupTextureParams(UInterchangeTextureFactor
 	{
 		if (CompressionSettings == "mask")
 			TextureFactoryNode->SetCustomCompressionSettings(TC_Masks);
-		else if (CompressionSettings == "displacement")
-			TextureFactoryNode->SetCustomCompressionSettings(TC_Displacementmap);
+		else if (CompressionSettings == "displacement" || CompressionSettings == "alpha")
+			TextureFactoryNode->SetCustomCompressionSettings(TC_Alpha);
 	}
 
 	if (FString MigGenSettings; TextureParams->TryGetStringField(TEXT("mipgen"), MigGenSettings))
@@ -529,9 +545,9 @@ void UInterchangeMegascansPipeline::SetupMaterial(UInterchangeMaterialInstanceFa
 			UMaterialInstanceConstant::StaticClass(),
 			"BasePropertyOverrides.bOverride_BlendMode"
 		);
-		MaterialInstanceFactoryNode->SetAttribute(MEGASCAN_MATERIAL_BLEND_MODE_KEY, static_cast<int>(BLEND_Opaque));
+		MaterialInstanceFactoryNode->SetAttribute(MEGASCAN_MATERIAL_BLEND_MODE_VALUE_KEY, static_cast<int>(BLEND_Opaque));
 		MaterialInstanceFactoryNode->AddApplyAndFillDelegates<int>(
-			MEGASCAN_MATERIAL_BLEND_MODE_KEY,
+			MEGASCAN_MATERIAL_BLEND_MODE_VALUE_KEY,
 			MaterialInstanceFactoryNode->GetObjectClass(),
 			"BasePropertyOverrides.BlendMode"
 		);
@@ -540,6 +556,35 @@ void UInterchangeMegascansPipeline::SetupMaterial(UInterchangeMaterialInstanceFa
 
 void UInterchangeMegascansPipeline::SetupMaterialParams(UInterchangeMaterialInstanceFactoryNode* MaterialInstanceFactoryNode, const TSharedPtr<FJsonObject>& MaterialParams)
 {
+	if (const TSharedPtr<FJsonObject>* Overrides; MaterialParams->TryGetObjectField(TEXT("overrides"), Overrides))
+	{
+		if (const TSharedPtr<FJsonObject>* Displacement; Overrides->Get()->TryGetObjectField(TEXT("displacement"), Displacement))
+		{
+			MaterialInstanceFactoryNode->SetAttribute(MEGASCAN_MATERIAL_DISPLACEMENT_OVERRIDE_KEY, true);
+			MaterialInstanceFactoryNode->AddApplyAndFillDelegates<bool>(
+				MEGASCAN_MATERIAL_DISPLACEMENT_OVERRIDE_KEY,
+				UMaterialInstanceConstant::StaticClass(),
+				"BasePropertyOverrides.bOverride_DisplacementScaling"
+			);
+
+			const float Magnitude = Displacement->Get()->GetNumberField(TEXT("magnitude"));
+			MaterialInstanceFactoryNode->SetAttribute(MEGASCAN_MATERIAL_DISPLACEMENT_MAGNITUDE_KEY, Magnitude);
+			MaterialInstanceFactoryNode->AddApplyAndFillDelegates<float>(
+				MEGASCAN_MATERIAL_DISPLACEMENT_MAGNITUDE_KEY,
+				UMaterialInstanceConstant::StaticClass(),
+				"BasePropertyOverrides.DisplacementScaling.Magnitude"
+			);
+
+			const float Center = Displacement->Get()->GetNumberField(TEXT("center"));
+			MaterialInstanceFactoryNode->SetAttribute(MEGASCAN_MATERIAL_DISPLACEMENT_CENTER_KEY, Center);
+			MaterialInstanceFactoryNode->AddApplyAndFillDelegates<float>(
+				MEGASCAN_MATERIAL_DISPLACEMENT_CENTER_KEY,
+				UMaterialInstanceConstant::StaticClass(),
+				"BasePropertyOverrides.DisplacementScaling.Center"
+			);
+		}
+	}
+
 	const TArray<TSharedPtr<FJsonValue>>* CustomParams;
 	if (!MaterialParams->TryGetArrayField(TEXT("params"), CustomParams))
 	{
@@ -589,49 +634,52 @@ void UInterchangeMegascansPipeline::SetupMaterialParents(UInterchangeMaterialIns
 	{
 		case EMegascanImportType::Model3D:
 		{
-			if (CustomType.IsEmpty())
-				SetMegascanMaterialType(MaterialInstanceFactoryNode, EMegascanMaterialType::Base3D);
+			if (CustomType.IsEmpty() || CustomType == "base")
+				SetMegascanMaterialType(MaterialInstanceFactoryNode, EMegascanMaterialType::Base);
+			else if (CustomType == "masked")
+				SetMegascanMaterialType(MaterialInstanceFactoryNode, EMegascanMaterialType::BaseMasked);
 			else if (CustomType == "transmission")
-				SetMegascanMaterialType(MaterialInstanceFactoryNode, EMegascanMaterialType::Transmission3D);
+				SetMegascanMaterialType(MaterialInstanceFactoryNode, EMegascanMaterialType::BaseTransmission);
 			else if (CustomType == "fuzz")
-				SetMegascanMaterialType(MaterialInstanceFactoryNode, EMegascanMaterialType::Fuzz3D);
+				SetMegascanMaterialType(MaterialInstanceFactoryNode, EMegascanMaterialType::BaseFuzz);
 			else if (CustomType == "glass")
-				SetMegascanMaterialType(MaterialInstanceFactoryNode, EMegascanMaterialType::Glass3D);
+				SetMegascanMaterialType(MaterialInstanceFactoryNode, EMegascanMaterialType::Glass);
 		}
 		break;
 		case EMegascanImportType::Surface:
 		{
-			if (CustomType.IsEmpty())
+			if (CustomType.IsEmpty() || CustomType == "surface")
 				SetMegascanMaterialType(MaterialInstanceFactoryNode, EMegascanMaterialType::Surface);
+			else if (CustomType == "masked")
+				SetMegascanMaterialType(MaterialInstanceFactoryNode, EMegascanMaterialType::SurfaceMasked);
 			else if (CustomType == "transmission")
-				SetMegascanMaterialType(MaterialInstanceFactoryNode, EMegascanMaterialType::TransmissionSurface);
+				SetMegascanMaterialType(MaterialInstanceFactoryNode, EMegascanMaterialType::SurfaceTransmission);
+			else if (CustomType == "fuzz")
+				SetMegascanMaterialType(MaterialInstanceFactoryNode, EMegascanMaterialType::SurfaceFuzz);
+			else if (CustomType == "fabric" || CustomType == "fabric_opaque")
+				SetMegascanMaterialType(MaterialInstanceFactoryNode, EMegascanMaterialType::Fabric);
+			else if (CustomType == "fabric_masked")
+				SetMegascanMaterialType(MaterialInstanceFactoryNode, EMegascanMaterialType::FabricMasked);
 		}
 		break;
 		case EMegascanImportType::Decal:
 		{
-			SetMegascanMaterialType(MaterialInstanceFactoryNode, EMegascanMaterialType::Decal);
-		}
-		break;
-		case EMegascanImportType::Imperfection:
-		{
-			SetMegascanMaterialType(MaterialInstanceFactoryNode, EMegascanMaterialType::Surface);
+			if (CustomType.IsEmpty() || CustomType == "decal")
+				SetMegascanMaterialType(MaterialInstanceFactoryNode, EMegascanMaterialType::Decal);
 		}
 		break;
 		case EMegascanImportType::Plant:
 		{
-			if (CustomType.IsEmpty())
+			if (CustomType.IsEmpty() || CustomType == "plant")
 				SetMegascanMaterialType(MaterialInstanceFactoryNode, EMegascanMaterialType::Plant);
 			else if (CustomType == "billboard")
 				SetMegascanMaterialType(MaterialInstanceFactoryNode, EMegascanMaterialType::PlantBillboard);
 		}
 		break;
+
+		default:
+			break;
 	}
 
 	UpdateParentMaterial(MaterialInstanceFactoryNode);
-
-	// if (MegascanImportType == EMegascanImportType::Plant && ImportTier == EMegascanImportTier::Raw)
-	// {
-	// 	const FString ParameterName = UInterchangeShaderPortsAPI::MakeInputValueKey(TEXT("secondary Animation"));
-	// 	MaterialInstanceFactoryNode->AddBooleanAttribute(ParameterName, true);
-	// }
 }
diff --git a/Engine/Plugins/Fab/Source/Fab/Private/Pipelines/InterchangeMegascansPipeline.h b/Engine/Plugins/Fab/Source/Fab/Private/Pipelines/InterchangeMegascansPipeline.h
--- a/Engine/Plugins/Fab/Source/Fab/Private/Pipelines/InterchangeMegascansPipeline.h
+++ b/Engine/Plugins/Fab/Source/Fab/Private/Pipelines/InterchangeMegascansPipeline.h
@@ -52,18 +52,24 @@ enum class EMegascanMaterialType : int32
 {
 	Invalid = 0 UMETA(Hidden),
 
-	Base3D UMETA(DisplayName = "3D"),
-	Fuzz3D UMETA(DisplayName="3D with Fuzz"),
-	Transmission3D UMETA(DisplayName="3D with Tranmission"),
-	Glass3D UMETA(DisplayName="Glass 3D"),
+	Base UMETA(DisplayName = "3D"),
+	BaseMasked UMETA(DisplayName = "3D Masked"),
+	BaseFuzz UMETA(DisplayName="3D Fuzz"),
+	BaseTransmission UMETA(DisplayName="3D Transmission"),
+
+	Glass UMETA(DisplayName="Glass"),
 
 	Surface UMETA(DisplayName = "Surface"),
-	TransmissionSurface UMETA(DisplayName="Surface with Tranmission"),
-	FabricSurface UMETA(DisplayName="Fabric Surface"),
+	SurfaceMasked UMETA(DisplayName = "Surface Masked"),
+	SurfaceFuzz UMETA(DisplayName="Surface Fuzz"),
+	SurfaceTransmission UMETA(DisplayName="Surface Transmission"),
+
+	Fabric UMETA(DisplayName="Fabric"),
+	FabricMasked UMETA(DisplayName="Fabric Masked"),
 
 	Decal UMETA(DisplayName = "Decal"),
 
-	Plant UMETA(DisplayName = "Plant Base"),
+	Plant UMETA(DisplayName = "Plant"),
 	PlantBillboard UMETA(DisplayName = "Plant Billboard"),
 };
 
diff --git a/Engine/Plugins/Fab/Source/Fab/Private/Utilities/FabAssetsCache.cpp b/Engine/Plugins/Fab/Source/Fab/Private/Utilities/FabAssetsCache.cpp
--- a/Engine/Plugins/Fab/Source/Fab/Private/Utilities/FabAssetsCache.cpp
+++ b/Engine/Plugins/Fab/Source/Fab/Private/Utilities/FabAssetsCache.cpp
@@ -6,7 +6,6 @@
 #include "HAL/FileManager.h"
 #include "Misc/Paths.h"
 
-FString FFabAssetsCache::CacheLocation              = FPlatformProcess::UserTempDir() / FString("FabLibrary");
 int64 FFabAssetsCache::CacheExpirationTimeoutInDays = 10;
 
 FString SizeSuffix(const int64 InSize)
@@ -29,7 +28,7 @@ TArray<FString> FFabAssetsCache::GetCachedAssets()
 {
 	TArray<FString> CachedAssets;
 	IFileManager::Get().IterateDirectory(
-		*CacheLocation,
+		*GetCacheLocation(),
 		[&CachedAssets](const TCHAR* Path, bool bIsDirectory)-> bool
 		{
 			if (!bIsDirectory && FPaths::GetExtension(Path) == "zip")
@@ -46,7 +45,7 @@ int64 FFabAssetsCache::GetCacheSize()
 {
 	int64 CacheSize = 0;
 	IFileManager::Get().IterateDirectoryStatRecursively(
-		*CacheLocation,
+		*GetCacheLocation(),
 		[&CacheSize](const TCHAR* Path, const FFileStatData& Stat)-> bool
 		{
 			if (!Stat.bIsDirectory)
@@ -67,15 +66,15 @@ FText FFabAssetsCache::GetCacheSizeString()
 void FFabAssetsCache::ClearCache()
 {
 	IFileManager& FileManager = IFileManager::Get();
-	FileManager.DeleteDirectory(*CacheLocation, false, true);
-	FileManager.MakeDirectory(*CacheLocation);
+	FileManager.DeleteDirectory(*GetCacheLocation(), false, true);
+	FileManager.MakeDirectory(*GetCacheLocation());
 }
 
 bool FFabAssetsCache::IsCached(const FString& AssetId, int64 DownloadSize)
 {
 	IFileManager& FileManager     = IFileManager::Get();
 	const FString CachedFile      = GetCachedFile(AssetId);
-	FFileStatData CachedFileStats = FileManager.GetStatData(*CachedFile);
+	const FFileStatData CachedFileStats = FileManager.GetStatData(*CachedFile);
 
 	return (CachedFileStats.bIsValid && CachedFileStats.FileSize == DownloadSize && (FDateTime::Now() - CachedFileStats.ModificationTime).GetTotalDays() <
 		CacheExpirationTimeoutInDays);
@@ -83,12 +82,12 @@ bool FFabAssetsCache::IsCached(const FString& AssetId, int64 DownloadSize)
 
 FString FFabAssetsCache::GetCachedFile(const FString AssetId)
 {
-	return CacheLocation / AssetId;
+	return GetCacheLocation() / AssetId;
 }
 
 FString FFabAssetsCache::CacheAsset(const FString& DownloadedAssetPath)
 {
-	const FString CacheFilePath = CacheLocation / FPaths::GetCleanFilename(DownloadedAssetPath);
+	const FString CacheFilePath = GetCacheLocation() / FPaths::GetCleanFilename(DownloadedAssetPath);
 	IFileManager::Get().Move(*CacheFilePath, *DownloadedAssetPath, true, true);
 	return CacheFilePath;
 }
diff --git a/Engine/Plugins/Fab/Source/Fab/Private/Utilities/FabAssetsCache.h b/Engine/Plugins/Fab/Source/Fab/Private/Utilities/FabAssetsCache.h
--- a/Engine/Plugins/Fab/Source/Fab/Private/Utilities/FabAssetsCache.h
+++ b/Engine/Plugins/Fab/Source/Fab/Private/Utilities/FabAssetsCache.h
@@ -3,15 +3,15 @@
 #pragma once
 
 #include "CoreMinimal.h"
+#include "FabSettings.h"
 
 class FFabAssetsCache
 {
 private:
 	static int64 CacheExpirationTimeoutInDays;
-	static FString CacheLocation;
 
 public:
-	static FString GetCacheLocation() { return CacheLocation; }
+	static FString GetCacheLocation() { return GetDefault<UFabSettings>()->CacheDirectoryPath.Path; }
 	static TArray<FString> GetCachedAssets();
 	static int64 GetCacheSize();
 	static FText GetCacheSizeString();
diff --git a/Engine/Plugins/Fab/Source/Fab/Private/Utilities/QuixelAssetTypes.cpp b/Engine/Plugins/Fab/Source/Fab/Private/Utilities/QuixelAssetTypes.cpp
--- a/Engine/Plugins/Fab/Source/Fab/Private/Utilities/QuixelAssetTypes.cpp
+++ b/Engine/Plugins/Fab/Source/Fab/Private/Utilities/QuixelAssetTypes.cpp
@@ -5,7 +5,7 @@
 #include "JsonObjectConverter.h"
 #include "Misc/FileHelper.h"
 
-TTuple<FString, FString> FQuixelAssetTypes::ExtractMeta(const FString& JsonFile)
+TTuple<FString, FString> FQuixelAssetTypes::ExtractMeta(const FString& JsonFile, const FString& GltfFile)
 {
 	FString FileContent;
 	FFileHelper::LoadFileToString(FileContent, *JsonFile);
@@ -14,42 +14,115 @@ TTuple<FString, FString> FQuixelAssetTypes::ExtractMeta(const FString& JsonFile)
 	FJsonObjectConverter::JsonObjectStringToUStruct(FileContent, &AssetMetaDataJson);
 	const FString AssetId = AssetMetaDataJson.Id;
 
+	if (!GltfFile.IsEmpty() && AssetMetaDataJson.Displacement_Scale_Tier1 >= 0.0f && AssetMetaDataJson.Displacement_Bias_Tier1 >= 0.0f)
+	{
+		// Temporary till displacement values are integrated in gltf
+		const TSharedPtr<FJsonObject> DisplacementObject = MakeShareable(new FJsonObject);
+		DisplacementObject->SetNumberField(TEXT("magnitude"), AssetMetaDataJson.Displacement_Scale_Tier1);
+		DisplacementObject->SetNumberField(TEXT("center"), AssetMetaDataJson.Displacement_Bias_Tier1);
+		if (FString GltfFileData; FFileHelper::LoadFileToString(GltfFileData, *GltfFile))
+		{
+			TSharedPtr<FJsonObject> GltfJson = MakeShareable(new FJsonObject);
+			if (FJsonSerializer::Deserialize(TJsonReaderFactory<TCHAR>::Create(GltfFileData), GltfJson))
+			{
+				if (const TArray<TSharedPtr<FJsonValue>>* Materials; GltfJson->TryGetArrayField(TEXT("materials"), Materials))
+				{
+					for (const TSharedPtr<FJsonValue>& Material : *Materials)
+					{
+						const TSharedPtr<FJsonObject>& MaterialObject = Material->AsObject();
+						if (const TSharedPtr<FJsonObject>* Extras; MaterialObject->TryGetObjectField(TEXT("extras"), Extras))
+						{
+							if (const TSharedPtr<FJsonObject>* Overrides; Extras->Get()->TryGetObjectField(TEXT("overrides"), Overrides))
+							{
+								Overrides->Get()->SetObjectField(TEXT("displacement"), DisplacementObject);
+							}
+							else
+							{
+								TSharedPtr<FJsonObject> OverridesObject = MakeShareable(new FJsonObject);
+								OverridesObject->SetObjectField(TEXT("displacement"), DisplacementObject);
+								Extras->Get()->SetObjectField(TEXT("overrides"), OverridesObject);
+							}
+						}
+						else
+						{
+							TSharedPtr<FJsonObject> OverridesObject = MakeShareable(new FJsonObject);
+							OverridesObject->SetObjectField(TEXT("displacement"), DisplacementObject);
+							TSharedPtr<FJsonObject> ExtrasObject = MakeShareable(new FJsonObject);
+							ExtrasObject->SetObjectField(TEXT("overrides"), OverridesObject);
+
+							MaterialObject->SetObjectField(TEXT("extras"), ExtrasObject);
+						}
+					}
+				}
+
+				if (FString SerializedJson; FJsonSerializer::Serialize(GltfJson.ToSharedRef(), TJsonWriterFactory<TCHAR>::Create(&SerializedJson, 2)))
+				{
+					FFileHelper::SaveStringToFile(SerializedJson, *GltfFile);
+				}
+			}
+		}
+	}
+
 	if (AssetMetaDataJson.Categories.Num() == 0)
 	{
-		return {AssetId, ""};
+		return {
+			AssetId,
+			""
+		};
 	}
 
 	if (AssetMetaDataJson.Categories[0] == "3d")
 	{
-		return {AssetId, "3D"};
+		return {
+			AssetId,
+			"3D"
+		};
 	}
 
 	if (AssetMetaDataJson.Categories[0] == "surface")
 	{
-		return {AssetId, "Surfaces"};
+		return {
+			AssetId,
+			"Surfaces"
+		};
 	}
 
 	if (AssetMetaDataJson.Categories[0] == "3dplant")
 	{
-		return {AssetId, "Plants"};
+		return {
+			AssetId,
+			"Plants"
+		};
 	}
 
 	if (AssetMetaDataJson.Categories[0] == "atlas" && AssetMetaDataJson.Categories.Num() > 1)
 	{
 		if (AssetMetaDataJson.Categories[1] == "decals")
 		{
-			return {AssetId, "Decals"};
+			return {
+				AssetId,
+				"Decals"
+			};
 		}
 		if (AssetMetaDataJson.Categories[1] == "imperfections")
 		{
-			return {AssetId, "Imperfections"};
+			return {
+				AssetId,
+				"Imperfections"
+			};
 		}
 	}
 
 	if (AssetMetaDataJson.SemanticTags.Asset_Type == "decal")
 	{
-		return {AssetId, "Decals"};
+		return {
+			AssetId,
+			"Decals"
+		};
 	}
 
-	return {AssetId, ""};
+	return {
+		AssetId,
+		""
+	};
 }
diff --git a/Engine/Plugins/Fab/Source/Fab/Private/Utilities/QuixelAssetTypes.h b/Engine/Plugins/Fab/Source/Fab/Private/Utilities/QuixelAssetTypes.h
--- a/Engine/Plugins/Fab/Source/Fab/Private/Utilities/QuixelAssetTypes.h
+++ b/Engine/Plugins/Fab/Source/Fab/Private/Utilities/QuixelAssetTypes.h
@@ -28,10 +28,16 @@ struct FAssetMetaDataJson
 
 	UPROPERTY()
 	FSemanticTags SemanticTags;
+
+	UPROPERTY()
+	float Displacement_Bias_Tier1 = -1.0f;
+
+	UPROPERTY()
+	float Displacement_Scale_Tier1 = -1.0f;
 };
 
 class FQuixelAssetTypes
 {
 public:
-	static TTuple<FString, FString> ExtractMeta(const FString& JsonFile);
+	static TTuple<FString, FString> ExtractMeta(const FString& JsonFile, const FString& GltfFile = "");
 };
diff --git a/Engine/Plugins/Fab/Source/Fab/Private/Workflows/GenericImportWorkflow.cpp b/Engine/Plugins/Fab/Source/Fab/Private/Workflows/GenericImportWorkflow.cpp
--- a/Engine/Plugins/Fab/Source/Fab/Private/Workflows/GenericImportWorkflow.cpp
+++ b/Engine/Plugins/Fab/Source/Fab/Private/Workflows/GenericImportWorkflow.cpp
@@ -170,19 +170,21 @@ void FGenericImportWorkflow::OnContentDownloadComplete(const FFabDownloadRequest
 		}
 	}
 
-	const auto HasSameFilename = [](const FString& A, const FString& B) { return FPaths::GetBaseFilename(A) == FPaths::GetBaseFilename(B); };
+	const auto IsSameFilename = [](const FString& A, const FString& B) { return FPaths::GetBaseFilename(A) == FPaths::GetBaseFilename(B); };
 	TArray<FString> UpdatedImportFiles;
 	UpdatedImportFiles.Reserve(ImportFiles.Num());
 	for (FString& ImportFile : ImportFiles)
 	{
-		if (UpdatedImportFiles.ContainsByPredicate([&](const FString& Path) { return HasSameFilename(ImportFile, Path); }))
+		if (UpdatedImportFiles.ContainsByPredicate([&](const FString& Path) { return IsSameFilename(ImportFile, Path); }))
 		{
 			const FString Extension = FPaths::GetExtension(ImportFile);
 			const FString NewFile   = FPaths::GetBaseFilename(ImportFile, false) + '_' + Extension + '.' + Extension;
 			if (FileManager.MoveFile(*NewFile, *ImportFile))
 				ImportFile = NewFile;
 		}
-		UpdatedImportFiles.Add(ImportFile);
+
+		if (FPaths::FileExists(ImportFile))
+			UpdatedImportFiles.Add(ImportFile);
 	}
 
 	if (UpdatedImportFiles.IsEmpty())
diff --git a/Engine/Plugins/Fab/Source/Fab/Private/Workflows/QuixelImportWorkflow.cpp b/Engine/Plugins/Fab/Source/Fab/Private/Workflows/QuixelImportWorkflow.cpp
--- a/Engine/Plugins/Fab/Source/Fab/Private/Workflows/QuixelImportWorkflow.cpp
+++ b/Engine/Plugins/Fab/Source/Fab/Private/Workflows/QuixelImportWorkflow.cpp
@@ -129,7 +129,7 @@ void FQuixelImportWorkflow::ImportContent(const TArray<FString>& SourceFiles)
 	const FString SourceFile = SourceFiles[0];
 	const FString MetaFile   = SourceFiles[1];
 
-	auto [MegascanId, SubType] = FQuixelAssetTypes::ExtractMeta(MetaFile);
+	auto [MegascanId, SubType] = FQuixelAssetTypes::ExtractMeta(MetaFile, SourceFile);
 	const FString TierString   = ExtractTierNameFromFilename(SourceFile);
 
 	ImportLocation = "/Game/Fab/Megascans" / SubType / AssetName + '_' + MegascanId / TierString;
@@ -156,7 +156,7 @@ void FQuixelImportWorkflow::ImportContent(const TArray<FString>& SourceFiles)
 
 	if (SubType == "3D")
 	{
-		FQuixelGltfImporter::ImportGltf3DAsset(SourceFile, ImportLocation, TierString == "Raw", OnDone);
+		FQuixelGltfImporter::ImportGltf3DAsset(SourceFile, ImportLocation, OnDone);
 	}
 	else if (SubType == "Plants")
 	{
