Remove menu references to opening the Unreal Market place; replace with Fab.

#rnx

diff --git a/Engine/Plugins/Fab/Source/Fab/Private/FabBrowser.cpp b/Engine/Plugins/Fab/Source/Fab/Private/FabBrowser.cpp
--- a/Engine/Plugins/Fab/Source/Fab/Private/FabBrowser.cpp
+++ b/Engine/Plugins/Fab/Source/Fab/Private/FabBrowser.cpp
@@ -110,40 +110,47 @@ void FFabBrowser::SetupEntryPoints()
 		FCanExecuteAction()
 	);
 
-	FToolMenuSection& SaveSection = UToolMenus::Get()->ExtendMenu("ContentBrowser.Toolbar")->FindOrAddSection("Save");
-	FToolMenuEntry& ToolMenuEntry = SaveSection.AddEntry(
-		FToolMenuEntry::InitToolBarButton(
-			"OpenFabWindow",
-			InvokeTabAction,
-			FabLabel,
-			FabTooltip,
-			FSlateIcon(SlateStyleSet->GetStyleSetName(), FabToolbarIconName),
-			EUserInterfaceActionType::Button
-		)
-	);
-	ToolMenuEntry.InsertPosition.Position = EToolMenuInsertType::Last;
-	ToolMenuEntry.StyleNameOverride = "ContentBrowser.ToolBar.Buttons";
-
-	UToolMenu* WindowMenu = UToolMenus::Get()->ExtendMenu("MainFrame.MainMenu.Window");
-
-	FToolMenuSection* ContentSectionPtr = WindowMenu->FindSection("GetContent");
-	if (!ContentSectionPtr)
 	{
-		ContentSectionPtr = &WindowMenu->AddSection("GetContent", NSLOCTEXT("MainAppMenu", "GetContentHeader", "Get Content"));
+		FToolMenuSection& SaveSection = UToolMenus::Get()->ExtendMenu("ContentBrowser.Toolbar")->FindOrAddSection("Save");
+		FToolMenuEntry& ToolMenuEntry = SaveSection.AddEntry(
+			FToolMenuEntry::InitToolBarButton(
+				"OpenFabWindow",
+				InvokeTabAction,
+				FabLabel,
+				FabTooltip,
+				FSlateIcon(SlateStyleSet->GetStyleSetName(), FabToolbarIconName),
+				EUserInterfaceActionType::Button
+			)
+		);
+		ToolMenuEntry.InsertPosition.Position = EToolMenuInsertType::Last;
+		ToolMenuEntry.StyleNameOverride = "ContentBrowser.ToolBar.Buttons";
 	}
 
-	ContentSectionPtr->AddMenuEntry(
-		TEXT("OpenFabTab"),
-		LOCTEXT("OpenFabTab_Label", "Fab"),
-		LOCTEXT("OpenFabTab_Desc", "Opens the Fab Plugin."),
+	FToolMenuEntry FabMenuEntry = FToolMenuEntry::InitMenuEntry(
+		"OpenFabTab",
+		FabLabel,
+		FabTooltip,
 		FSlateIcon(SlateStyleSet->GetStyleSetName(), FabMenuIconName),
 		InvokeTabAction
 	);
 
+	{
+		UToolMenu* WindowMenu = UToolMenus::Get()->ExtendMenu("MainFrame.MainMenu.Window");
+		FToolMenuSection& ContentSection = WindowMenu->FindOrAddSection("GetContent", NSLOCTEXT("MainAppMenu", "GetContentHeader", "Get Content"));
+		FToolMenuEntry& FabEntry = ContentSection.AddEntry(FabMenuEntry);
+		FabEntry.InsertPosition.Position = EToolMenuInsertType::First;
+	}
+	
 	// Add a Fab entry to the Content Browser's Add popup menu
 	UToolMenus::Get()->ExtendMenu(TEXT("ContentBrowser.AddNewContextMenu"))->AddSection(TEXT("ContentBrowserGetContent"), LOCTEXT("GetContentText", "Get Content")).AddEntry(
 		FToolMenuEntry::InitMenuEntry(TEXT("OpenFabWindow"), FabLabel, FabTooltip, FSlateIcon(SlateStyleSet->GetStyleSetName(), FabMenuIconName), InvokeTabAction)
 	);
+	
+	{
+		FToolMenuEntry& FabEntry = UToolMenus::Get()->ExtendMenu("LevelEditor.LevelEditorToolBar.AddQuickMenu")->FindOrAddSection("Content").AddEntry(FabMenuEntry);
+		FabEntry.InsertPosition.Name = "ImportContent";
+		FabEntry.InsertPosition.Position = EToolMenuInsertType::After;
+	}
 }
 
 // Extend the context menu to view listings in Fab
diff --git a/Engine/Source/Editor/ContentBrowser/Private/ContentBrowserSingleton.cpp b/Engine/Source/Editor/ContentBrowser/Private/ContentBrowserSingleton.cpp
--- a/Engine/Source/Editor/ContentBrowser/Private/ContentBrowserSingleton.cpp
+++ b/Engine/Source/Editor/ContentBrowser/Private/ContentBrowserSingleton.cpp
@@ -125,7 +125,7 @@ FContentBrowserSingleton::FContentBrowserSingleton()
 	Section.AddSubMenu("ContentBrowser", LOCTEXT("ContentBrowserMenu", "Content Browser"), LOCTEXT("ContentBrowserTooltip", "Actions related to the Content Browser"),
 						FNewToolMenuDelegate::CreateRaw(this, &FContentBrowserSingleton::GetContentBrowserSubMenu, ContentBrowserGroup), false, 
 						FSlateIcon(FAppStyle::Get().GetStyleSetName(), "LevelEditor.OpenContentBrowser"))
-						.InsertPosition = FToolMenuInsert("OpenMarketplace", EToolMenuInsertType::After);
+						.InsertPosition.Position = EToolMenuInsertType::Last;
 
 	// Register to be notified when properties are edited
 	FEditorDelegates::LoadSelectedAssetsIfNeeded.AddRaw(this, &FContentBrowserSingleton::OnEditorLoadSelectedAssetsIfNeeded);
diff --git a/Engine/Source/Editor/GameProjectGeneration/Private/SProjectBrowser.cpp b/Engine/Source/Editor/GameProjectGeneration/Private/SProjectBrowser.cpp
--- a/Engine/Source/Editor/GameProjectGeneration/Private/SProjectBrowser.cpp
+++ b/Engine/Source/Editor/GameProjectGeneration/Private/SProjectBrowser.cpp
@@ -1247,49 +1247,6 @@ void SProjectBrowser::HandleProjectViewSelectionChanged(TSharedPtr<FProjectItem>
 	ProjectSelectionChangedDelegate.ExecuteIfBound(ProjectFile);
 }
 
-
-FReply SProjectBrowser::OnOpenMarketplace()
-{
-	ILauncherPlatform* LauncherPlatform = FLauncherPlatformModule::Get();
-
-	if (LauncherPlatform != nullptr)
-	{
-		TArray<FAnalyticsEventAttribute> EventAttributes;
-
-		FOpenLauncherOptions OpenOptions(TEXT("ue/marketplace"));
-		if (LauncherPlatform->OpenLauncher(OpenOptions) )
-		{
-			EventAttributes.Add(FAnalyticsEventAttribute(TEXT("OpenSucceeded"), TEXT("TRUE")));
-		}
-		else
-		{
-			EventAttributes.Add(FAnalyticsEventAttribute(TEXT("OpenSucceeded"), TEXT("FALSE")));
-
-			if (EAppReturnType::Yes == FMessageDialog::Open(EAppMsgType::YesNo, LOCTEXT("InstallMarketplacePrompt", "The Marketplace requires the Epic Games Launcher, which does not seem to be installed on your computer. Would you like to install it now?")))
-			{
-				FOpenLauncherOptions InstallOptions(true, TEXT("ue/marketplace"));
-				if (!LauncherPlatform->OpenLauncher(InstallOptions))
-				{
-					EventAttributes.Add(FAnalyticsEventAttribute(TEXT("InstallSucceeded"), TEXT("FALSE")));
-					FMessageDialog::Open(EAppMsgType::Ok, FText::FromString(TEXT("Sorry, there was a problem installing the Launcher.\nPlease try to install it manually!")));
-				}
-				else
-				{
-					EventAttributes.Add(FAnalyticsEventAttribute(TEXT("InstallSucceeded"), TEXT("TRUE")));
-				}
-			}
-		}
-
-		EventAttributes.Add(FAnalyticsEventAttribute(TEXT("Source"), TEXT("ProjectBrowser")));
-		if( FEngineAnalytics::IsAvailable() )
-		{
-			FEngineAnalytics::GetProvider().RecordEvent(TEXT("Editor.Usage.OpenMarketplace"), EventAttributes);
-		}
-	}
-
-	return FReply::Handled();
-}
-
 void SProjectBrowser::OnFilterTextChanged(const FText& InText)
 {
 	ProjectItemFilter.SetRawFilterText(InText);
diff --git a/Engine/Source/Editor/GameProjectGeneration/Private/SProjectBrowser.h b/Engine/Source/Editor/GameProjectGeneration/Private/SProjectBrowser.h
--- a/Engine/Source/Editor/GameProjectGeneration/Private/SProjectBrowser.h
+++ b/Engine/Source/Editor/GameProjectGeneration/Private/SProjectBrowser.h
@@ -91,9 +91,6 @@ public:
 
 	FReply OnBrowseToProject();
 
-	/** Callback for clicking the 'Marketplace' button. */
-	FReply OnOpenMarketplace();
-
 	/** Handler for when the Open button is clicked */
 	FReply OnOpenProject();
 
diff --git a/Engine/Source/Editor/GameProjectGeneration/Private/SProjectDialog.cpp b/Engine/Source/Editor/GameProjectGeneration/Private/SProjectDialog.cpp
--- a/Engine/Source/Editor/GameProjectGeneration/Private/SProjectDialog.cpp
+++ b/Engine/Source/Editor/GameProjectGeneration/Private/SProjectDialog.cpp
@@ -594,37 +594,6 @@ TSharedRef<SWidget> SProjectDialog::MakeOpenProjectDialogButtons()
 	return
 		SNew(SHorizontalBox)
 		+ SHorizontalBox::Slot()
-		.Padding(8.0f, 0.0f, 16.0f, 0.0f)
-		.VAlign(VAlign_Center)
-		.AutoWidth()
-		[
-			SNew(SButton)
-			.ToolTipText(LOCTEXT("MarketplaceToolTip", "Check out the Marketplace to find new projects!"))
-			.Visibility(FLauncherPlatformModule::Get()->CanOpenLauncher(true) ? EVisibility::Visible : EVisibility::Collapsed)
-			.OnClicked(ProjectBrowserRef, &SProjectBrowser::OnOpenMarketplace)
-			[	
-				SNew(SHorizontalBox)
-				+SHorizontalBox::Slot()
-				.Padding(4.0f,0.0f)
-				.AutoWidth()
-				.HAlign(HAlign_Left)
-				.VAlign(VAlign_Center)
-				[
-					SNew(SImage)
-					.Image(FAppStyle::Get().GetBrush("LevelEditor.OpenMarketplace"))
-					.DesiredSizeOverride(FVector2D(16,16))
-					.ColorAndOpacity(FSlateColor::UseForeground())
-				]
-				+ SHorizontalBox::Slot()
-				.Padding(4.0f, 0.0f)
-				.VAlign(VAlign_Center)
-				[
-					SNew(STextBlock)
-					.Text(LOCTEXT("OpenMarketplace", "Marketplace..."))
-				]
-			]
-		]
-		+ SHorizontalBox::Slot()
 		.Padding(8.0f, 0.0f)
 		.VAlign(VAlign_Center)
 		.AutoWidth()
@@ -633,6 +602,7 @@ TSharedRef<SWidget> SProjectDialog::MakeOpenProjectDialogButtons()
 			.HAlign(HAlign_Center)
 			.VAlign(VAlign_Center)
 			.Text(LOCTEXT("BrowseForProjects", "Browse..."))
+			.ToolTipText(LOCTEXT("BrowseForProjects_Tooltip", "Browse to and open a project on your computer."))
 			.OnClicked(ProjectBrowserRef, &SProjectBrowser::OnBrowseToProject)
 		]
 		+SHorizontalBox::Slot()
@@ -643,6 +613,7 @@ TSharedRef<SWidget> SProjectDialog::MakeOpenProjectDialogButtons()
 			SNew(SPrimaryButton)
 			.Visibility(this, &SProjectDialog::GetCreateButtonVisibility)
 			.Text(LOCTEXT("OpenProject", "Open"))
+			.ToolTipText(LOCTEXT("OpenProject_Tooltip", "Open the selected project."))
 			.IsEnabled(ProjectBrowserRef, &SProjectBrowser::HasSelectedProjectFile)
 			.OnClicked(ProjectBrowserRef, &SProjectBrowser::OnOpenProject)
 		]
diff --git a/Engine/Source/Editor/LevelEditor/Private/LevelEditorActions.cpp b/Engine/Source/Editor/LevelEditor/Private/LevelEditorActions.cpp
--- a/Engine/Source/Editor/LevelEditor/Private/LevelEditorActions.cpp
+++ b/Engine/Source/Editor/LevelEditor/Private/LevelEditorActions.cpp
@@ -2623,11 +2623,6 @@ void FLevelEditorActionCallbacks::OpenContentBrowser()
 	ContentBrowserModule.Get().FocusPrimaryContentBrowser(true);
 }
 
-void FLevelEditorActionCallbacks::OpenMarketplace()
-{
-	FUnrealEdMisc::Get().OpenMarketplace();
-}
-
 void FLevelEditorActionCallbacks::ImportContent()
 {
 	FString Path = "/Game";
@@ -4353,7 +4348,6 @@ void FLevelEditorCommands::RegisterCommands()
 	UI_COMMAND( WorldProperties, "World Settings", "Displays the world settings", EUserInterfaceActionType::Button, FInputChord() );
 	UI_COMMAND( OpenPlaceActors, "Place Actors Panel", "Opens the Place Actors Panel", EUserInterfaceActionType::Button, FInputChord() );
 	UI_COMMAND( OpenContentBrowser, "Open Content Browser", "Opens the Content Browser", EUserInterfaceActionType::Button, FInputChord(EModifierKey::Control|EModifierKey::Shift, EKeys::F) );
-	UI_COMMAND( OpenMarketplace, "Unreal Marketplace", "Opens the Marketplace", EUserInterfaceActionType::Button, FInputChord() );
 	UI_COMMAND( ImportContent, "Import Content...", "Import Content into a specified location", EUserInterfaceActionType::Button, FInputChord());
 
 	UI_COMMAND( ToggleVR, "Toggle VR", "Toggles VR (Virtual Reality) mode", EUserInterfaceActionType::ToggleButton, FInputChord(EModifierKey::Shift, EKeys::V ) );
diff --git a/Engine/Source/Editor/LevelEditor/Private/LevelEditorToolBar.cpp b/Engine/Source/Editor/LevelEditor/Private/LevelEditorToolBar.cpp
--- a/Engine/Source/Editor/LevelEditor/Private/LevelEditorToolBar.cpp
+++ b/Engine/Source/Editor/LevelEditor/Private/LevelEditorToolBar.cpp
@@ -2020,11 +2020,6 @@ void FLevelEditorToolBar::RegisterAddMenu()
 		Section.InitSection("Content", LOCTEXT("Content_Label", "Get Content"), FToolMenuInsert(NAME_None, EToolMenuInsertType::First));
 
 		Section.AddMenuEntry(FLevelEditorCommands::Get().ImportContent).InsertPosition.Position = EToolMenuInsertType::First;
-
-		if (FLauncherPlatformModule::Get()->CanOpenLauncher(true))
-		{
-			Section.AddMenuEntry(FLevelEditorCommands::Get().OpenMarketplace);
-		}
 	}
 #undef LOCTEXT_NAMESPACE
 }
diff --git a/Engine/Source/Editor/LevelEditor/Private/SLevelEditor.cpp b/Engine/Source/Editor/LevelEditor/Private/SLevelEditor.cpp
--- a/Engine/Source/Editor/LevelEditor/Private/SLevelEditor.cpp
+++ b/Engine/Source/Editor/LevelEditor/Private/SLevelEditor.cpp
@@ -159,10 +159,6 @@ void SLevelEditor::BindCommands()
 		Actions.OpenContentBrowser,
 		FExecuteAction::CreateStatic( &FLevelEditorActionCallbacks::OpenContentBrowser ) );
 	
-	LevelEditorCommands->MapAction(
-		Actions.OpenMarketplace,
-		FExecuteAction::CreateStatic( &FLevelEditorActionCallbacks::OpenMarketplace ) );
-
 	LevelEditorCommands->MapAction(
 		Actions.ImportContent,
 		FExecuteAction::CreateStatic(&FLevelEditorActionCallbacks::ImportContent));
diff --git a/Engine/Source/Editor/LevelEditor/Public/LevelEditorActions.h b/Engine/Source/Editor/LevelEditor/Public/LevelEditorActions.h
--- a/Engine/Source/Editor/LevelEditor/Public/LevelEditorActions.h
+++ b/Engine/Source/Editor/LevelEditor/Public/LevelEditorActions.h
@@ -566,7 +566,6 @@ public:
 	TSharedPtr< FUICommandInfo > WorldProperties;
 	TSharedPtr< FUICommandInfo > OpenPlaceActors;
 	TSharedPtr< FUICommandInfo > OpenContentBrowser;
-	TSharedPtr< FUICommandInfo > OpenMarketplace;
 	TSharedPtr< FUICommandInfo > ToggleVR;
 	TSharedPtr< FUICommandInfo > ImportContent;
 
@@ -1279,9 +1278,6 @@ public:
 	/** Open the Content Browser */
 	static void OpenContentBrowser();
 
-	/** Open the Marketplace */
-	static void OpenMarketplace();
-
 	/** Import content into a chosen location*/
 	static void ImportContent();
 
diff --git a/Engine/Source/Editor/MainFrame/Private/Frame/MainFrameActions.cpp b/Engine/Source/Editor/MainFrame/Private/Frame/MainFrameActions.cpp
--- a/Engine/Source/Editor/MainFrame/Private/Frame/MainFrameActions.cpp
+++ b/Engine/Source/Editor/MainFrame/Private/Frame/MainFrameActions.cpp
@@ -184,9 +184,6 @@ void FMainFrameCommands::RegisterCommands()
 												FCanExecuteAction(),
 												FIsActionChecked::CreateStatic( &FMainFrameActionCallbacks::OpenSlateApp_IsChecked, FName("SessionFrontend" ) ) );
 
-	UI_COMMAND(OpenMarketplace, "Open Marketplace", "Opens the Marketplace", EUserInterfaceActionType::Button, FInputChord());
-	ActionList->MapAction(OpenMarketplace, FExecuteAction::CreateStatic(&FMainFrameActionCallbacks::OpenMarketplace));
-
 	UI_COMMAND(DocumentationHome, "Documentation Home", "Authoritative, in-depth technical resources for using Unreal Engine", EUserInterfaceActionType::Button, FInputChord());
 	ActionList->MapAction(DocumentationHome, FExecuteAction::CreateStatic(&FMainFrameActionCallbacks::DocumentationHome));
 
@@ -380,11 +377,6 @@ void FMainFrameActionCallbacks::CacheProjectNames()
 	RecentProjects = GetDefault<UEditorSettings>()->RecentlyOpenedProjectFiles;
 }
 
-void FMainFrameActionCallbacks::OpenMarketplace()
-{
-	FUnrealEdMisc::Get().OpenMarketplace();
-}
-
 void FMainFrameActionCallbacks::NewProject( bool bAllowProjectOpening, bool bAllowProjectCreate )
 {
 	if (GUnrealEd->WarnIfLightingBuildIsCurrentlyRunning())
diff --git a/Engine/Source/Editor/MainFrame/Private/Frame/MainFrameActions.h b/Engine/Source/Editor/MainFrame/Private/Frame/MainFrameActions.h
--- a/Engine/Source/Editor/MainFrame/Private/Frame/MainFrameActions.h
+++ b/Engine/Source/Editor/MainFrame/Private/Frame/MainFrameActions.h
@@ -51,7 +51,6 @@ public:
 	TSharedPtr< FUICommandInfo > OpenDebugView;
 	TSharedPtr< FUICommandInfo > OpenClassViewer;
 	TSharedPtr< FUICommandInfo > OpenWidgetReflector;
-	TSharedPtr< FUICommandInfo > OpenMarketplace;
 
 	TSharedPtr< FUICommandInfo > DocumentationHome;
 	TSharedPtr< FUICommandInfo > BrowseAPIReference;
@@ -278,8 +277,6 @@ public:
 	/** Gathers all available projects the user can switch to from main menu */
 	static void CacheProjectNames();
 
-	/** Opens the marketplace */
-	static void OpenMarketplace();
 public:
 
 	// List of projects that the user can switch to.
diff --git a/Engine/Source/Editor/MainFrame/Private/Menus/MainMenu.cpp b/Engine/Source/Editor/MainFrame/Private/Menus/MainMenu.cpp
--- a/Engine/Source/Editor/MainFrame/Private/Menus/MainMenu.cpp
+++ b/Engine/Source/Editor/MainFrame/Private/Menus/MainMenu.cpp
@@ -319,15 +319,7 @@ void FMainMenu::RegisterWindowMenu()
 	}));
 
 	// Get content section
-	{
-		if (FLauncherPlatformModule::Get()->CanOpenLauncher(true))
-		{
-			FToolMenuSection& Section = Menu->AddSection("GetContent", NSLOCTEXT("MainAppMenu", "GetContentHeader", "Get Content"));
-			Section.AddMenuEntry(FMainFrameCommands::Get().OpenMarketplace);
-		}
-		
-	}
-
+	Menu->AddSection("GetContent", NSLOCTEXT("MainAppMenu", "GetContentHeader", "Get Content"));
 
 	// Layout section
 	{
