[FNameColumn] Add a new column to the TEDS Fab integration to represent the name column and replace uses of the FNameColumn

diff --git a/Engine/Plugins/Fab/Source/Fab/Private/Teds/FabContentSource.cpp b/Engine/Plugins/Fab/Source/Fab/Private/Teds/FabContentSource.cpp
--- a/Engine/Plugins/Fab/Source/Fab/Private/Teds/FabContentSource.cpp
+++ b/Engine/Plugins/Fab/Source/Fab/Private/Teds/FabContentSource.cpp
@@ -5,7 +5,6 @@
 #include "FabBrowser.h"
 #include "FabMyFolderIntegration.h"
 #include "IContentBrowserSingleton.h"
-#include "Elements/Columns/TypedElementMiscColumns.h"
 #include "Elements/Columns/TypedElementWebColumns.h"
 #include "Elements/Framework/TypedElementQueryBuilder.h"
 
@@ -63,7 +62,7 @@ void FFabMyLibraryContentSource::GetAssetViewInitParams(UE::Editor::ContentBrows
 				.All<FFabObjectColumn>()
 			.Compile();
 
-	OutInitParams.Columns = { FNameColumn::StaticStruct(), FUrlColumn::StaticStruct() };
+	OutInitParams.Columns = { FFabObjectNameColumn::StaticStruct(), FUrlColumn::StaticStruct() };
 
 	OutInitParams.CellWidgetPurpose = IUiProvider::FPurposeID(IUiProvider::FPurposeInfo(
 		"General", "Cell", NAME_None).GeneratePurposeID());
diff --git a/Engine/Plugins/Fab/Source/Fab/Private/Teds/FabMyFolderIntegration.cpp b/Engine/Plugins/Fab/Source/Fab/Private/Teds/FabMyFolderIntegration.cpp
--- a/Engine/Plugins/Fab/Source/Fab/Private/Teds/FabMyFolderIntegration.cpp
+++ b/Engine/Plugins/Fab/Source/Fab/Private/Teds/FabMyFolderIntegration.cpp
@@ -3,7 +3,6 @@
 #include "Teds/FabMyFolderIntegration.h"
 
 #include "Dom/JsonObject.h"
-#include "Elements/Columns/TypedElementMiscColumns.h"
 #include "Elements/Columns/TypedElementWebColumns.h"
 #include "Elements/Common/EditorDataStorageFeatures.h"
 #include "Elements/Interfaces/TypedElementDataStorageInterface.h"
@@ -19,7 +18,7 @@
 
 void UFabFactory::RegisterTables(UE::Editor::DataStorage::ICoreProvider& DataStorage)
 {
-	Table = DataStorage.RegisterTable<FNameColumn, FFabObjectColumn, FEditorDataStorageUrlColumn>(FName("Fab"));
+	Table = DataStorage.RegisterTable<FFabObjectNameColumn, FFabObjectColumn, FEditorDataStorageUrlColumn>(FName("Fab"));
 }
 
 void FFabTedsMyFolderIntegration::QueueSyncRequest()
@@ -129,7 +128,7 @@ void FFabTedsMyFolderIntegration::ProcessSyncResults(uint32 BatchSize, HttpReque
 								const TSharedPtr<FJsonObject>* ResultObject = nullptr;
 								if (bool Found = Result->TryGetObject(ResultObject); Found && ResultObject->IsValid())
 								{
-									if (FNameColumn* Target = Storage->GetColumn<FNameColumn>(Row))
+									if (FFabObjectNameColumn* Target = Storage->GetColumn<FFabObjectNameColumn>(Row))
 									{
 										SetNameColumn(*Target, **ResultObject, Temp);
 									}
@@ -186,7 +185,7 @@ void FFabTedsMyFolderIntegration::ProcessSyncResults(uint32 BatchSize, HttpReque
 	}
 }
 
-void FFabTedsMyFolderIntegration::SetNameColumn(FNameColumn& Target, const FJsonObject& Object, FString& Temp)
+void FFabTedsMyFolderIntegration::SetNameColumn(FFabObjectNameColumn& Target, const FJsonObject& Object, FString& Temp)
 {
 	if (Object.TryGetStringField(TEXT("title"), Temp))
 	{
diff --git a/Engine/Plugins/Fab/Source/Fab/Private/Teds/FabMyFolderIntegration.h b/Engine/Plugins/Fab/Source/Fab/Private/Teds/FabMyFolderIntegration.h
--- a/Engine/Plugins/Fab/Source/Fab/Private/Teds/FabMyFolderIntegration.h
+++ b/Engine/Plugins/Fab/Source/Fab/Private/Teds/FabMyFolderIntegration.h
@@ -10,7 +10,6 @@
 
 struct FEditorDataStorageUrlColumn;
 class FJsonObject;
-struct FNameColumn;
 
 namespace UE::Editor::DataStorage
 {
@@ -38,6 +37,15 @@ struct FFabDistributionMethodTag : public FEditorDataStorageTag
 	struct EditorDataStorage_DynamicColumnTemplate{};
 };
 
+USTRUCT(meta = (DisplayName = "Name"))
+struct FFabObjectNameColumn : public FEditorDataStorageColumn
+{
+	GENERATED_BODY()
+
+	UPROPERTY(meta = (Searchable, Sortable))
+	FName Name;
+};
+
 USTRUCT(meta = (DisplayName = "Fab object"))
 struct FFabObjectColumn : public FEditorDataStorageColumn
 {
@@ -78,7 +86,7 @@ private:
 	static void QueueSyncRequest(const FString& Cursor, uint32 BatchSize);
 	static void ProcessSyncResults(uint32 BatchSize, HttpRequestPtr Request, HttpResponsePtr Response, bool bWasSuccessful);
 
-	static void SetNameColumn(FNameColumn& Target, const FJsonObject& Object, FString& Temp);
+	static void SetNameColumn(FFabObjectNameColumn& Target, const FJsonObject& Object, FString& Temp);
 	static void SetFabObjectColumn(FFabObjectColumn& Target, const FJsonObject& Object, FString& Temp);
 	static void SetUrlColumn(FEditorDataStorageUrlColumn& Target, const FJsonObject& Object, FString& Temp);
 	static void AddDistributionMethod(
