Fix for Fab crashing when a project has non-ascii characters in its name. Validate the EOS platform handle before trying to use it as it fails to init EOS in this case.

diff --git a/Engine/Plugins/Fab/Source/Fab/Private/FabAuthentication.cpp b/Engine/Plugins/Fab/Source/Fab/Private/FabAuthentication.cpp
--- a/Engine/Plugins/Fab/Source/Fab/Private/FabAuthentication.cpp
+++ b/Engine/Plugins/Fab/Source/Fab/Private/FabAuthentication.cpp
@@ -97,6 +97,10 @@ namespace FabAuthentication
 
 			PlatformHandle = SDKManager->CreatePlatform(PlatformOptions);
 		}
+		else
+		{
+			FAB_LOG_ERROR("EOS is not initialized.");
+		}
 	}
 
 	void Shutdown()
@@ -243,6 +247,12 @@ namespace FabAuthentication
 	{
 		FAB_LOG("Logging in using exchange code");
 
+		if (!PlatformHandle)
+		{
+			FAB_LOG_ERROR("Invalid EOS platform handle.");
+			return false;
+		}
+
 		if (ExchangeCode.IsEmpty()) // Read exchange code from commandline if it is not passed in
 		{
 			FAB_LOG("Reading exchange code from commandline");
@@ -277,6 +287,12 @@ namespace FabAuthentication
 	{
 		FAB_LOG("Logging in using persist");
 
+		if (!PlatformHandle)
+		{
+			FAB_LOG_ERROR("Invalid EOS platform handle.");
+			return false;
+		}
+
 		AuthHandle = EOS_Platform_GetAuthInterface(*PlatformHandle);
 
 		EOS_Auth_Credentials Credentials = {0};
@@ -297,6 +313,12 @@ namespace FabAuthentication
 	{
 		FAB_LOG("Logging in using account portal");
 
+		if (!PlatformHandle)
+		{
+			FAB_LOG_ERROR("Invalid EOS platform handle.");
+			return false;
+		}
+
 		AuthHandle = EOS_Platform_GetAuthInterface(*PlatformHandle);
 
 		EOS_Auth_Credentials Credentials = {0};
