fixed UE5 Fab editor plugin logout loop on macOS WebKit and CEF

diff --git a/Engine/Plugins/Fab/Source/Fab/Private/FabAuthentication.cpp b/Engine/Plugins/Fab/Source/Fab/Private/FabAuthentication.cpp
--- a/Engine/Plugins/Fab/Source/Fab/Private/FabAuthentication.cpp
+++ b/Engine/Plugins/Fab/Source/Fab/Private/FabAuthentication.cpp
@@ -380,12 +380,26 @@ namespace FabAuthentication
 		}
 	}
 
+	void OnLogoutCallbackFn(const EOS_Auth_LogoutCallbackInfo* Data) 
+	{
+		if(Data->ResultCode != EOS_EResult::EOS_Success)
+		{
+			FAB_LOG_ERROR("EOS_Auth_Logout failed");
+		}
+	}
+
 	void DeletePersistentAuth()
 	{
 		FAB_LOG("Delete persist auth");
 
 		EOS_Auth_DeletePersistentAuthOptions Options = {};
+		EOS_Auth_LogoutOptions LogoutOptions = {};
+		
 		Options.ApiVersion = EOS_AUTH_DELETEPERSISTENTAUTH_API_LATEST;
+		LogoutOptions.ApiVersion = EOS_AUTH_LOGOUT_API_LATEST;
+		LogoutOptions.LocalUserId = EpicAccountId;
+		
+		EOS_Auth_Logout(AuthHandle, &LogoutOptions, nullptr, OnLogoutCallbackFn);
 		EOS_Auth_DeletePersistentAuth(AuthHandle, &Options, nullptr, DeletePersistentAuthCompleteCallbackFn);
 	}
 
diff --git a/Engine/Plugins/Fab/Source/Fab/Private/FabAuthentication.h b/Engine/Plugins/Fab/Source/Fab/Private/FabAuthentication.h
--- a/Engine/Plugins/Fab/Source/Fab/Private/FabAuthentication.h
+++ b/Engine/Plugins/Fab/Source/Fab/Private/FabAuthentication.h
@@ -154,4 +154,6 @@ namespace FabAuthentication
 	* Callback that is fired when the delete persistent auth operation completes, either successfully or in error
 	*/
 	static void EOS_CALL DeletePersistentAuthCompleteCallbackFn(const EOS_Auth_DeletePersistentAuthCallbackInfo* Data);
+
+	static void EOS_CALL OnLogoutCallbackFn(const EOS_Auth_LogoutCallbackInfo* Data);
 }
