fix BugHawk issues for CL #44629137

diff --git a/Engine/Plugins/Fab/Source/Fab/Private/FabAuthentication.cpp b/Engine/Plugins/Fab/Source/Fab/Private/FabAuthentication.cpp
--- a/Engine/Plugins/Fab/Source/Fab/Private/FabAuthentication.cpp
+++ b/Engine/Plugins/Fab/Source/Fab/Private/FabAuthentication.cpp
@@ -380,27 +380,31 @@ namespace FabAuthentication
 		}
 	}
 
-	void OnLogoutCallbackFn(const EOS_Auth_LogoutCallbackInfo* Data) 
+	void EOS_CALL OnLogoutCallbackFn(const EOS_Auth_LogoutCallbackInfo* Data) 
 	{
 		if(Data->ResultCode != EOS_EResult::EOS_Success)
 		{
 			FAB_LOG_ERROR("EOS_Auth_Logout failed");
 		}
+		else
+		{
+			EOS_Auth_DeletePersistentAuthOptions Options = {};
+			Options.ApiVersion = EOS_AUTH_DELETEPERSISTENTAUTH_API_LATEST;
+			
+			EOS_Auth_DeletePersistentAuth(AuthHandle, &Options, nullptr, DeletePersistentAuthCompleteCallbackFn);
+		}
 	}
 
 	void DeletePersistentAuth()
 	{
 		FAB_LOG("Delete persist auth");
 
-		EOS_Auth_DeletePersistentAuthOptions Options = {};
 		EOS_Auth_LogoutOptions LogoutOptions = {};
 		
-		Options.ApiVersion = EOS_AUTH_DELETEPERSISTENTAUTH_API_LATEST;
 		LogoutOptions.ApiVersion = EOS_AUTH_LOGOUT_API_LATEST;
 		LogoutOptions.LocalUserId = EpicAccountId;
 		
 		EOS_Auth_Logout(AuthHandle, &LogoutOptions, nullptr, OnLogoutCallbackFn);
-		EOS_Auth_DeletePersistentAuth(AuthHandle, &Options, nullptr, DeletePersistentAuthCompleteCallbackFn);
 	}
 
 	FString GetAuthToken()
