[Fab] Lower case uproperties for macOS WebKit compatibility - fixes MH issue

diff --git a/Engine/Plugins/Fab/Source/Fab/Private/FabBrowserApi.cpp b/Engine/Plugins/Fab/Source/Fab/Private/FabBrowserApi.cpp
--- a/Engine/Plugins/Fab/Source/Fab/Private/FabBrowserApi.cpp
+++ b/Engine/Plugins/Fab/Source/Fab/Private/FabBrowserApi.cpp
@@ -323,18 +323,18 @@ FFabApiVersion UFabBrowserApi::GetApiVersion()
 	const uint16 MajorVersion = EngineVersion.GetMajor();
 	const uint16 MinorVersion = EngineVersion.GetMinor();
 
-	ApiVersion.Ue = FString::FromInt(MajorVersion) + "." + FString::FromInt(MinorVersion);
-	ApiVersion.Api = "1.0.0";
+	ApiVersion.ue = FString::FromInt(MajorVersion) + "." + FString::FromInt(MinorVersion);
+	ApiVersion.api = "1.0.0";
 
 	if (const TSharedPtr<IPlugin> Plugin = IPluginManager::Get().FindPlugin("Fab"); Plugin.IsValid())
 	{
 		const FPluginDescriptor& PluginDescriptor = Plugin->GetDescriptor();
 
 		const FString PluginVersion = PluginDescriptor.VersionName;
-		ApiVersion.PluginVersion = PluginVersion;
+		ApiVersion.pluginversion = PluginVersion;
 	}
 
-	ApiVersion.Platform = UGameplayStatics::GetPlatformName();
+	ApiVersion.platform = UGameplayStatics::GetPlatformName();
 
 	return ApiVersion;
 }
diff --git a/Engine/Plugins/Fab/Source/Fab/Private/FabBrowserApi.h b/Engine/Plugins/Fab/Source/Fab/Private/FabBrowserApi.h
--- a/Engine/Plugins/Fab/Source/Fab/Private/FabBrowserApi.h
+++ b/Engine/Plugins/Fab/Source/Fab/Private/FabBrowserApi.h
@@ -13,17 +13,20 @@ struct FFabApiVersion
 {
 	GENERATED_BODY()
 
+	// These three uproperties need to be in this case (NOT PascalCase)
+	// otherwise we will break compatibility with the engine plugin 
+	// Fab's JS on the macOS WebKit WebBrowser plugin backend.
 	UPROPERTY()
-	FString Ue;
+	FString ue;
 
 	UPROPERTY()
-	FString Api;
+	FString api;
 
 	UPROPERTY()
-	FString PluginVersion;
+	FString pluginversion;
 
 	UPROPERTY()
-	FString Platform;
+	FString platform;
 };
 
 USTRUCT()
