[Fab] Bump plugin version for UE 5.7

diff --git a/Engine/Plugins/Fab/Fab.uplugin b/Engine/Plugins/Fab/Fab.uplugin
--- a/Engine/Plugins/Fab/Fab.uplugin
+++ b/Engine/Plugins/Fab/Fab.uplugin
@@ -1,7 +1,7 @@
 {
 	"FileVersion": 3,
-	"Version": 5,
-	"VersionName": "0.0.5",
+	"Version": 7,
+	"VersionName": "0.0.7",
 	"FriendlyName": "Fab",
 	"Description": "Fab Plugin",
 	"Category": "Other",
