Fab: Change FabDownloader HTTP Requests to use a stream delegate to download files, fix for downloading big files (2GB+).
(Crash in HTTP!UE::Core::Private::ReallocGrow3())

diff --git a/Engine/Plugins/Fab/Source/Fab/Private/FabDownloader.cpp b/Engine/Plugins/Fab/Source/Fab/Private/FabDownloader.cpp
--- a/Engine/Plugins/Fab/Source/Fab/Private/FabDownloader.cpp
+++ b/Engine/Plugins/Fab/Source/Fab/Private/FabDownloader.cpp
@@ -4,6 +4,7 @@
 
 #include "FabLog.h"
 #include "HttpModule.h"
+#include "HAL/PlatformFileManager.h"
 
 #include "Importers/BuildPatchInstallerLibHelper.h"
 
@@ -81,7 +82,8 @@ void FFabDownloadRequest::ShutdownBpsModule()
 void FFabDownloadRequest::ExecuteHTTPRequest()
 {
 	const FString FullFileName = GetFilenameFromURL(DownloadURL);
-	const FString SaveFilename = DownloadLocation / GetFilenameFromURL(DownloadURL);
+	const FString SaveFilename = DownloadLocation / FullFileName;
+	const FString DownloadFile = SaveFilename + TEXT(".download");
 
 	FHttpModule& HTTPModule = FHttpModule::Get();
 
@@ -92,7 +94,7 @@ void FFabDownloadRequest::ExecuteHTTPRequest()
 		{
 			if (HeaderName == "Content-Length")
 			{
-				DownloadStats.TotalBytes    = FCString::Atoi(*HeaderValue);
+				DownloadStats.TotalBytes    = FCString::Atoi64(*HeaderValue);
 				const FString CachedAssetId = AssetID / FullFileName;
 				OnDownloadProgressDelegate.Broadcast(this, DownloadStats);
 
@@ -111,40 +113,102 @@ void FFabDownloadRequest::ExecuteHTTPRequest()
 			}
 		}
 	);
-	#if (ENGINE_MAJOR_VERSION >=5 && ENGINE_MINOR_VERSION <=3)
+
+	// we need to create dir path to store the downloaded file
+	const FString Directory = FPaths::GetPath(DownloadFile);
+	IPlatformFile& PlatformFile = FPlatformFileManager::Get().GetPlatformFile();
+	if (!PlatformFile.DirectoryExists(*Directory))
+	{
+		if (!PlatformFile.CreateDirectoryTree(*Directory))
+		{
+			FAB_LOG_ERROR("Failed to create the dir for writing: %s", *Directory);
+			DownloadStats.bIsSuccess = false;
+			OnDownloadCompleteDelegate.Broadcast(this, DownloadStats);
+			return;
+		}
+	}
+
+	// stream directly to file
+	struct FFileHandleWrapper { TUniquePtr<IFileHandle> FileHandle; };
+	TSharedPtr<FFileHandleWrapper> StreamToFileHandle = MakeShared<FFileHandleWrapper>(TUniquePtr<IFileHandle>(PlatformFile.OpenWrite(*DownloadFile)));
+	if (StreamToFileHandle == nullptr || StreamToFileHandle->FileHandle == nullptr)
+	{
+		FAB_LOG_ERROR("Failed to open file for writing: %s", *DownloadFile);
+		DownloadStats.bIsSuccess = false;
+		OnDownloadCompleteDelegate.Broadcast(this, DownloadStats);
+		StreamToFileHandle.Reset();
+
+		return;
+	}
+
+	DownloadRequest->SetResponseBodyReceiveStreamDelegateV2(FHttpRequestStreamDelegateV2::CreateLambda(
+	[this, StreamToFileHandle](const void* DataPtr, int64& DataLength)
+	{
+		if (StreamToFileHandle != nullptr && StreamToFileHandle->FileHandle != nullptr && DataLength > 0)
+		{
+			if (!StreamToFileHandle->FileHandle->Write(static_cast<const uint8*>(DataPtr), DataLength))
+			{
+				// signals error to the HTTP module by setting the data length zero
+				DataLength = 0;
+				FAB_LOG_ERROR("Failed to write data to file");
+				StreamToFileHandle->FileHandle.Reset();
+			}
+		}
+	}));
+
+#if (ENGINE_MAJOR_VERSION >=5 && ENGINE_MINOR_VERSION <=3)
 	DownloadRequest->OnRequestProgress().BindLambda(
 		[this](FHttpRequestPtr Request, uint32 UploadedBytes, uint32 DownloadedBytes)
-		#else
+#else
 	DownloadRequest->OnRequestProgress64().BindLambda(
 		[this](FHttpRequestPtr Request, uint64 UploadedBytes, uint64 DownloadedBytes)
-		#endif
+#endif
 		{
-			DownloadStats.CompletedBytes  = static_cast<uint64>(DownloadedBytes);
-			DownloadStats.PercentComplete = 100.0f * static_cast<float>(DownloadStats.CompletedBytes) / static_cast<float>(DownloadStats.TotalBytes);
-
+			DownloadStats.CompletedBytes  = DownloadedBytes;
+			DownloadStats.PercentComplete = DownloadStats.TotalBytes > 0
+				? 100.f * static_cast<float>(DownloadStats.CompletedBytes) / static_cast<float>(DownloadStats.TotalBytes)
+				: 0.f;
 			// TODO: Calculate Download Speed
 			OnDownloadProgressDelegate.Broadcast(this, DownloadStats);
 		}
 	);
 	DownloadRequest->OnProcessRequestComplete().BindLambda(
-		[this, SaveFilename](FHttpRequestPtr Request, FHttpResponsePtr Response, const bool bRequestComplete)
+	[this, SaveFilename, DownloadFile, StreamToFileHandle]
+	(FHttpRequestPtr Request, FHttpResponsePtr Response, const bool bRequestComplete)
+	{
+		bool bFileInDisk = false;
+		if (StreamToFileHandle != nullptr && StreamToFileHandle->FileHandle != nullptr)
 		{
-			if (bRequestComplete)
+			bFileInDisk = StreamToFileHandle->FileHandle->Flush(true);
+			// Close file handle now:
+			StreamToFileHandle->FileHandle.Reset();
+		}
+
+		IPlatformFile& PlatformFile = FPlatformFileManager::Get().GetPlatformFile();
+		const int32 ResponseCode = Response.IsValid() ? Response->GetResponseCode() : 0;
+		DownloadStats.bIsSuccess = false; // false unless we can download and rename
+		if (bFileInDisk && bRequestComplete && ResponseCode >= 200 && ResponseCode < 300)
+		{
+			// Download OK: let's rename the download to the actual final filename
+			if (PlatformFile.MoveFile(*SaveFilename, *DownloadFile))
 			{
-				const TArray<uint8>& Data         = Response->GetContent();
-				DownloadStats.bIsSuccess          = FFileHelper::SaveArrayToFile(Data, *SaveFilename);
 				DownloadStats.DownloadCompletedAt = FDateTime::Now().ToUnixTimestamp();
-				if (DownloadStats.bIsSuccess)
-				{
-					DownloadStats.DownloadedFiles = {
-						SaveFilename
-					};
-				}
+				DownloadStats.DownloadedFiles = { SaveFilename };
+				DownloadStats.bIsSuccess = true;
 			}
+			else
+			{
+				FAB_LOG_ERROR("Failed to rename downloaded file");
+			}
+		}
 
-			OnDownloadCompleteDelegate.Broadcast(this, DownloadStats);
+		if (!DownloadStats.bIsSuccess)
+		{
+			PlatformFile.DeleteFile(*DownloadFile);
 		}
-	);
+
+		OnDownloadCompleteDelegate.Broadcast(this, DownloadStats);
+	});
 
 	DownloadStats.DownloadStartedAt = FDateTime::Now().ToUnixTimestamp();
 	DownloadRequest->ProcessRequest();
