[Fab] Fix caching when downloading Fab listings across projects

#Synced-CL 48761024

diff --git a/Engine/Plugins/Fab/Source/Fab/Private/FabDownloader.cpp b/Engine/Plugins/Fab/Source/Fab/Private/FabDownloader.cpp
--- a/Engine/Plugins/Fab/Source/Fab/Private/FabDownloader.cpp
+++ b/Engine/Plugins/Fab/Source/Fab/Private/FabDownloader.cpp
@@ -103,6 +103,7 @@ void FFabDownloadRequest::ExecuteHTTPRequest()
 					DownloadStats.PercentComplete     = 100.0f;
 					DownloadStats.DownloadCompletedAt = FDateTime::Now().ToUnixTimestamp();
 					DownloadStats.bIsSuccess          = true;
+					DownloadStats.bIsCached           = true;
 					DownloadStats.CompletedBytes      = DownloadStats.TotalBytes;
 					DownloadStats.DownloadedFiles     = {
 						FFabAssetsCache::GetCachedFile(CachedAssetId)
@@ -176,6 +177,7 @@ void FFabDownloadRequest::ExecuteHTTPRequest()
 	[this, SaveFilename, DownloadFile, StreamToFileHandle]
 	(FHttpRequestPtr Request, FHttpResponsePtr Response, const bool bRequestComplete)
 	{
+		IPlatformFile& PlatformFile = FPlatformFileManager::Get().GetPlatformFile();
 		bool bFileInDisk = false;
 		if (StreamToFileHandle != nullptr && StreamToFileHandle->FileHandle != nullptr)
 		{
@@ -184,7 +186,13 @@ void FFabDownloadRequest::ExecuteHTTPRequest()
 			StreamToFileHandle->FileHandle.Reset();
 		}
 
-		IPlatformFile& PlatformFile = FPlatformFileManager::Get().GetPlatformFile();
+		if (DownloadStats.bIsCached)
+		{
+			PlatformFile.DeleteFile(*DownloadFile);
+			OnDownloadCompleteDelegate.Broadcast(this, DownloadStats);
+			return;
+		}
+
 		const int32 ResponseCode = Response.IsValid() ? Response->GetResponseCode() : 0;
 		DownloadStats.bIsSuccess = false; // false unless we can download and rename
 		if (bFileInDisk && bRequestComplete && ResponseCode >= 200 && ResponseCode < 300)
diff --git a/Engine/Plugins/Fab/Source/Fab/Public/FabDownloader.h b/Engine/Plugins/Fab/Source/Fab/Public/FabDownloader.h
--- a/Engine/Plugins/Fab/Source/Fab/Public/FabDownloader.h
+++ b/Engine/Plugins/Fab/Source/Fab/Public/FabDownloader.h
@@ -40,6 +40,7 @@ struct FFabDownloadStats
 	float DownloadSpeed = 0.0f;
 
 	bool bIsSuccess = false;
+	bool bIsCached = false;
 
 	TArray<FString> DownloadedFiles;
 };
