[Fab] Plugin version 0.0.10 - with fixes for downloads, d&d crashes and includes better logs

#Synced-CL 50982074

#Synced-CL 50982074

diff --git a/Engine/Plugins/Fab/Fab.uplugin b/Engine/Plugins/Fab/Fab.uplugin
--- a/Engine/Plugins/Fab/Fab.uplugin
+++ b/Engine/Plugins/Fab/Fab.uplugin
@@ -1,23 +1,34 @@
 {
 	"FileVersion": 3,
-	"Version": 7,
-	"VersionName": "0.0.7",
+	"Version": 10,
+	"VersionName": "0.0.10",
 	"FriendlyName": "Fab",
 	"Description": "Fab Plugin",
 	"Category": "Other",
 	"CreatedBy": "Epic Games, Inc.",
-	"CreatedByURL": "http://epicgames.com",
+	"CreatedByURL": "https://epicgames.com",
 	"DocsURL": "",
 	"MarketplaceURL": "",
 	"SupportURL": "",
-	"CanContainContent": true,
 	"EnabledByDefault": true,
+	"CanContainContent": true,
+	"Modules": [
+		{
+			"Name": "Fab",
+			"Type": "Editor",
+			"LoadingPhase": "PostDefault",
+			"PlatformAllowList": [
+				"Linux",
+				"Win64",
+				"Mac"
+			]
+		}
+	],
 	"Plugins": [
 		{
 			"Name": "Interchange",
 			"Enabled": true,
-			"PlatformAllowList" :
-			[
+			"PlatformAllowList": [
 				"Linux",
 				"Win64",
 				"Mac"
@@ -34,27 +45,14 @@
 		{
 			"Name": "EOSShared",
 			"Enabled": true,
-			"TargetAllowList": [
-				"Editor"
-			],
 			"PlatformAllowList": [
 				"Linux",
 				"Win64",
 				"Mac"
-			]
-		}
-	],
-	"Modules": [
-		{
-			"Name": "Fab",
-			"Type": "Editor",
-			"LoadingPhase": "PostDefault",
-			"PlatformAllowList" :
-			[
-				"Linux",
-				"Win64",
-				"Mac"
+			],
+			"TargetAllowList": [
+				"Editor"
 			]
 		}
 	]
-}
+}
\ No newline at end of file
diff --git a/Engine/Plugins/Fab/Source/Fab/Private/FabBrowserApi.cpp b/Engine/Plugins/Fab/Source/Fab/Private/FabBrowserApi.cpp
--- a/Engine/Plugins/Fab/Source/Fab/Private/FabBrowserApi.cpp
+++ b/Engine/Plugins/Fab/Source/Fab/Private/FabBrowserApi.cpp
@@ -47,8 +47,8 @@ void UFabBrowserApi::AddToProject(const FString& DownloadUrl, const FFabAssetMet
 		}
 	}
 
-	FAB_LOG("Asset Type = %s", *AssetMetadata.AssetType);
-	FAB_LOG("Is Quixel = %d", AssetMetadata.IsQuixel);
+	FAB_LOG("Add listing to project - %s", *AssetMetadata.AssetId);
+	FAB_LOG("Listing type - %s", *AssetMetadata.AssetType);
 
 	if (AssetMetadata.AssetType == "unreal-engine")
 	{
@@ -177,8 +177,9 @@ void UFabBrowserApi::DragStart(const FFabAssetMetadata& AssetMetadata)
 		}
 	}
 
-	FAB_LOG("Listing Type = %s", *AssetMetadata.ListingType);
-	FAB_LOG("Is Quixel = %d", AssetMetadata.IsQuixel);
+	FAB_LOG("Drag listing to project - %s", *AssetMetadata.AssetId);
+	FAB_LOG("Listing type - %s", *AssetMetadata.AssetType);
+
 	if (AssetMetadata.IsQuixel)
 	{
 		const TSharedPtr<FQuixelDragDropWorkflow> QuixelDragDropWorkflow = MakeShared<FQuixelDragDropWorkflow>(
diff --git a/Engine/Plugins/Fab/Source/Fab/Private/FabModule.cpp b/Engine/Plugins/Fab/Source/Fab/Private/FabModule.cpp
--- a/Engine/Plugins/Fab/Source/Fab/Private/FabModule.cpp
+++ b/Engine/Plugins/Fab/Source/Fab/Private/FabModule.cpp
@@ -13,6 +13,7 @@
 
 #include "PropertyEditorModule.h"
 #include "Engine/RendererSettings.h"
+#include "Interfaces/IPluginManager.h"
 #include "Modules/ModuleManager.h"
 
 #include "Pipelines/Factories/InterchangeInstancedFoliageTypeFactory.h"
@@ -68,6 +69,20 @@ public:
 			{
 				FCoreDelegates::OnPostEngineInit.AddLambda(RegisterItems);
 			}
+
+			const TSharedPtr<IPlugin> Plugin = IPluginManager::Get().FindPlugin(TEXT("Fab"));
+
+			if (Plugin.IsValid())
+			{
+				const FPluginDescriptor& Descriptor = Plugin->GetDescriptor();
+
+				UE_LOG(LogFab, Log,
+					TEXT("Plugin initialized. VersionName: %s (Version: %d) - Engine Version %s"),
+					*Descriptor.VersionName,
+					Descriptor.Version,
+					*FEngineVersion::Current().ToString()
+				);
+			}
 		}
 	}
 
diff --git a/Engine/Plugins/Fab/Source/Fab/Private/Importers/FabDragDropOp.cpp b/Engine/Plugins/Fab/Source/Fab/Private/Importers/FabDragDropOp.cpp
--- a/Engine/Plugins/Fab/Source/Fab/Private/Importers/FabDragDropOp.cpp
+++ b/Engine/Plugins/Fab/Source/Fab/Private/Importers/FabDragDropOp.cpp
@@ -113,11 +113,12 @@ void FFabDragDropOp::DestroyWindow()
 
 void FFabDragDropOp::DestroySpawnedActor()
 {
-	if (SpawnedActor && (DragAssetType == EDragAssetType::Mesh || DragAssetType == EDragAssetType::Decal))
+	if (SpawnedActor.IsValid() && (DragAssetType == EDragAssetType::Mesh || DragAssetType == EDragAssetType::Decal))
 	{
 		SpawnedActor->Destroy();
-		SpawnedActor = nullptr;
 	}
+
+	SpawnedActor = nullptr;
 }
 
 void FFabDragDropOp::Cancel()
diff --git a/Engine/Plugins/Fab/Source/Fab/Private/NotificationProgressWidget.h b/Engine/Plugins/Fab/Source/Fab/Private/NotificationProgressWidget.h
--- a/Engine/Plugins/Fab/Source/Fab/Private/NotificationProgressWidget.h
+++ b/Engine/Plugins/Fab/Source/Fab/Private/NotificationProgressWidget.h
@@ -71,7 +71,7 @@ public:
 				.HAlign(HAlign_Right)
 				.VAlign(VAlign_Top)
 				[
-					SNew(SButton)
+					SAssignNew(CancelButton, SButton)
 					.Text(InArgs._ButtonText)
 					.IsEnabled(InArgs._HasButton)
 					.ToolTipText(InArgs._ButtonToolTip)
@@ -88,6 +88,11 @@ public:
 		PercentText->SetColorAndOpacity(Percent <= 55.0f ? FLinearColor::White : FLinearColor::Black);
 	}
 
+	void DisableCancel()
+	{
+		CancelButton->SetEnabled(false);
+	}
+
 	/** INotificationWidget interface */
 	virtual void OnSetCompletionState(SNotificationItem::ECompletionState) override
 	{}
@@ -100,4 +105,5 @@ public:
 private:
 	TSharedPtr<SProgressBar> ProgressBar;
 	TSharedPtr<STextBlock> PercentText;
+	TSharedPtr<SButton> CancelButton;
 };
diff --git a/Engine/Plugins/Fab/Source/Fab/Private/Utilities/DragImportOperation.cpp b/Engine/Plugins/Fab/Source/Fab/Private/Utilities/DragImportOperation.cpp
--- a/Engine/Plugins/Fab/Source/Fab/Private/Utilities/DragImportOperation.cpp
+++ b/Engine/Plugins/Fab/Source/Fab/Private/Utilities/DragImportOperation.cpp
@@ -117,12 +117,13 @@ void FDragImportOperation::CancelOperation()
 	FSlateApplication::Get().CancelDragDrop();
 }
 
-AActor* FDragImportOperation::GetSpawnedActor() const
+TWeakObjectPtr<AActor> FDragImportOperation::GetSpawnedActor() const
 {
 	if (DragOperationHandle)
 	{
 		return DragOperationHandle->SpawnedActor;
 	}
+
 	return nullptr;
 }
 
@@ -144,24 +145,34 @@ void FDragImportOperation::ReplaceSpawnedActor() const
 
 	UObject* NewObject = DraggedAsset.GetAsset();
 	bool bIsReplaced   = false;
+
+	const TWeakObjectPtr<AActor>& SpawnedActor = GetSpawnedActor();
+	if (!SpawnedActor.IsValid())
+	{
+		DragOperationHandle->DestroySpawnedActor();
+		return;
+	}
+
 	if (DragAssetType == EDragAssetType::Mesh)
 	{
 		if (UStaticMesh* SourceMesh = Cast<UStaticMesh>(NewObject))
 		{
-			if (AStaticMeshActor* SMActor = Cast<AStaticMeshActor>(GetSpawnedActor()))
+			if (AStaticMeshActor* SMActor = Cast<AStaticMeshActor>(SpawnedActor))
 			{
-				SMActor->GetStaticMeshComponent()->EmptyOverrideMaterials();
-				SMActor->GetStaticMeshComponent()->SetStaticMesh(SourceMesh);
-				SMActor->SetActorLabel(DraggedAsset.AssetName.ToString());
-				bIsReplaced = true;
+				if (SMActor->GetStaticMeshComponent())
+				{
+					SMActor->GetStaticMeshComponent()->EmptyOverrideMaterials();
+					SMActor->GetStaticMeshComponent()->SetStaticMesh(SourceMesh);
+					SMActor->SetActorLabel(DraggedAsset.AssetName.ToString());
+					bIsReplaced = true;
+				}
 			}
 		}
 		else if (USkeletalMesh* SourceSkeletalMesh = Cast<USkeletalMesh>(NewObject))
 		{
-			ASkeletalMeshActor* SkMActor = Cast<ASkeletalMeshActor>(GetSpawnedActor());
+			ASkeletalMeshActor* SkMActor = Cast<ASkeletalMeshActor>(SpawnedActor);
 			if (SkMActor == nullptr)
 			{
-				AActor* SpawnedActor = GetSpawnedActor();
 				SkMActor             = SpawnedActor->GetWorld()->SpawnActor<ASkeletalMeshActor>(
 					ASkeletalMeshActor::StaticClass(),
 					SpawnedActor->GetTransform()
@@ -170,23 +181,29 @@ void FDragImportOperation::ReplaceSpawnedActor() const
 					DragOperationHandle->SpawnedActor = SkMActor;
 				SpawnedActor->Destroy();
 			}
-			SkMActor->GetSkeletalMeshComponent()->EmptyOverrideMaterials();
-			SkMActor->GetSkeletalMeshComponent()->SetSkeletalMesh(SourceSkeletalMesh);
-			SkMActor->SetActorLabel(DraggedAsset.AssetName.ToString());
-			bIsReplaced = true;
+			if (SkMActor->GetSkeletalMeshComponent())
+			{
+				SkMActor->GetSkeletalMeshComponent()->EmptyOverrideMaterials();
+				SkMActor->GetSkeletalMeshComponent()->SetSkeletalMesh(SourceSkeletalMesh);
+				SkMActor->SetActorLabel(DraggedAsset.AssetName.ToString());
+				bIsReplaced = true;
+			}
 		}
 	}
 	else if (DragAssetType == EDragAssetType::Material)
 	{
-		if (GetSpawnedActor())
+		UMaterialInterface* Material = Cast<UMaterialInterface>(NewObject);
+		if (const AStaticMeshActor* SMActor = Cast<AStaticMeshActor>(SpawnedActor))
 		{
-			UMaterialInterface* Material = Cast<UMaterialInterface>(NewObject);
-			if (const AStaticMeshActor* SMActor = Cast<AStaticMeshActor>(GetSpawnedActor()))
+			if (SMActor->GetStaticMeshComponent())
 			{
 				SMActor->GetStaticMeshComponent()->SetMaterial(0, Material);
 				bIsReplaced = true;
 			}
-			else if (const ASkeletalMeshActor* SkMActor = Cast<ASkeletalMeshActor>(GetSpawnedActor()))
+		}
+		else if (const ASkeletalMeshActor* SkMActor = Cast<ASkeletalMeshActor>(SpawnedActor))
+		{
+			if (SkMActor->GetSkeletalMeshComponent())
 			{
 				SkMActor->GetSkeletalMeshComponent()->SetMaterial(0, Material);
 				bIsReplaced = true;
@@ -197,7 +214,7 @@ void FDragImportOperation::ReplaceSpawnedActor() const
 	{
 		if (UMaterialInterface* Material = Cast<UMaterialInterface>(NewObject))
 		{
-			if (ADecalActor* DecalActor = Cast<ADecalActor>(GetSpawnedActor()))
+			if (ADecalActor* DecalActor = Cast<ADecalActor>(SpawnedActor))
 			{
 				DecalActor->SetActorLabel(DraggedAsset.AssetName.ToString());
 				DecalActor->SetDecalMaterial(Material);
diff --git a/Engine/Plugins/Fab/Source/Fab/Private/Utilities/DragImportOperation.h b/Engine/Plugins/Fab/Source/Fab/Private/Utilities/DragImportOperation.h
--- a/Engine/Plugins/Fab/Source/Fab/Private/Utilities/DragImportOperation.h
+++ b/Engine/Plugins/Fab/Source/Fab/Private/Utilities/DragImportOperation.h
@@ -29,7 +29,7 @@ public:
 
 	void CancelOperation();
 
-	AActor* GetSpawnedActor() const;
+	TWeakObjectPtr<AActor> GetSpawnedActor() const;
 	void DeleteSpawnedActor() const;
 	void ReplaceSpawnedActor() const;
 
diff --git a/Engine/Plugins/Fab/Source/Fab/Private/Workflows/GenericDragDropWorkflow.cpp b/Engine/Plugins/Fab/Source/Fab/Private/Workflows/GenericDragDropWorkflow.cpp
--- a/Engine/Plugins/Fab/Source/Fab/Private/Workflows/GenericDragDropWorkflow.cpp
+++ b/Engine/Plugins/Fab/Source/Fab/Private/Workflows/GenericDragDropWorkflow.cpp
@@ -72,9 +72,12 @@ void FGenericDragDropWorkflow::OnContentDownloadProgress(const FFabDownloadReque
 	{
 		return;
 	}
-	if (DragOperation && DragOperation->GetSpawnedActor())
+	if (DragOperation)
 	{
-		if (UMeshComponent* const MeshComponent = DragOperation->GetSpawnedActor()->GetComponentByClass<UMeshComponent>())
+		const TWeakObjectPtr<AActor>& SpawnedActor = DragOperation->GetSpawnedActor();
+		if (!SpawnedActor.IsValid()) return;
+
+		if (UMeshComponent* const MeshComponent = SpawnedActor->GetComponentByClass<UMeshComponent>())
 		{
 			if (MeshComponent->GetNumOverrideMaterials() == 0)
 			{
diff --git a/Engine/Plugins/Fab/Source/Fab/Private/Workflows/PackImportWorkflow.cpp b/Engine/Plugins/Fab/Source/Fab/Private/Workflows/PackImportWorkflow.cpp
--- a/Engine/Plugins/Fab/Source/Fab/Private/Workflows/PackImportWorkflow.cpp
+++ b/Engine/Plugins/Fab/Source/Fab/Private/Workflows/PackImportWorkflow.cpp
@@ -86,8 +86,15 @@ void FPackImportWorkflow::CreateDownloadNotification()
 				[this]()
 				{
 					FAB_LOG("Import Cancelled");
+					if (ProgressWidget.IsValid())
+					{
+						FAB_LOG("Disabling cancel button");
+						ProgressWidget->DisableCancel();
+					}	
 					if (DownloadRequest)
+					{
 						DownloadRequest->Cancel();
+					}
 					return FReply::Handled();
 				}
 			)
diff --git a/Engine/Plugins/Fab/Source/Fab/Private/Workflows/QuixelDragDropWorkflow.cpp b/Engine/Plugins/Fab/Source/Fab/Private/Workflows/QuixelDragDropWorkflow.cpp
--- a/Engine/Plugins/Fab/Source/Fab/Private/Workflows/QuixelDragDropWorkflow.cpp
+++ b/Engine/Plugins/Fab/Source/Fab/Private/Workflows/QuixelDragDropWorkflow.cpp
@@ -125,9 +125,12 @@ void FQuixelDragDropWorkflow::OnContentDownloadProgress(const FFabDownloadReques
 	{
 		return;
 	}
-	if (DragOperation && DragOperation->GetSpawnedActor())
+	if (DragOperation)
 	{
-		if (UMeshComponent* const MeshComponent = DragOperation->GetSpawnedActor()->GetComponentByClass<UMeshComponent>())
+		const TWeakObjectPtr<AActor>& SpawnedActor = DragOperation->GetSpawnedActor();
+		if (!SpawnedActor.IsValid()) return;
+
+		if (UMeshComponent* const MeshComponent = SpawnedActor->GetComponentByClass<UMeshComponent>())
 		{
 			if (MeshComponent->GetNumOverrideMaterials() == 0)
 			{
diff --git a/Engine/Plugins/Fab/Source/Fab/Public/Importers/FabDragDropOp.h b/Engine/Plugins/Fab/Source/Fab/Public/Importers/FabDragDropOp.h
--- a/Engine/Plugins/Fab/Source/Fab/Public/Importers/FabDragDropOp.h
+++ b/Engine/Plugins/Fab/Source/Fab/Public/Importers/FabDragDropOp.h
@@ -23,10 +23,8 @@ public:
 	DECLARE_DELEGATE(FOnDrop);
 	DRAG_DROP_OPERATOR_TYPE(FFabDragDropOp, FAssetDragDropOp)
 
-public:
 	static UE_API TSharedPtr<FFabDragDropOp> New(FAssetData Asset, EDragAssetType InDragAssetType);
 
-public:
 	UE_API FFabDragDropOp(const EDragAssetType InDragAssetType);
 	UE_API virtual ~FFabDragDropOp() override;
 
@@ -47,8 +45,7 @@ public:
 	UE_API virtual void OnDragged(const FDragDropEvent& DragDropEvent) override;
 	UE_API virtual void OnDrop(bool bDropWasHandled, const FPointerEvent& MouseEvent) override;
 
-public:
-	TObjectPtr<AActor> SpawnedActor;
+	TWeakObjectPtr<AActor> SpawnedActor;
 
 protected:
 	FOnDrop OnDropDelegate;
