[Fab] Plugin version 0.0.12 with fix for 3D previewer URL and pack import cancel safety

 - Fix navigation interception remapping fab.com/dope/ URLs, which broke the 3D previewer
 - Use WeakPtr capture in PackImportWorkflow cancel button lambda to prevent dangling pointer if notification outlives the workflow
 - Bump plugin version to 0.0.12

#Synced-CL 52249030

diff --git a/Engine/Plugins/Fab/Fab.uplugin b/Engine/Plugins/Fab/Fab.uplugin
--- a/Engine/Plugins/Fab/Fab.uplugin
+++ b/Engine/Plugins/Fab/Fab.uplugin
@@ -1,7 +1,7 @@
 {
 	"FileVersion": 3,
-	"Version": 11,
-	"VersionName": "0.0.11",
+	"Version": 12,
+	"VersionName": "0.0.12",
 	"FriendlyName": "Fab",
 	"Description": "Fab Plugin",
 	"Category": "Other",
diff --git a/Engine/Plugins/Fab/Source/Fab/Private/FabBrowser.cpp b/Engine/Plugins/Fab/Source/Fab/Private/FabBrowser.cpp
--- a/Engine/Plugins/Fab/Source/Fab/Private/FabBrowser.cpp
+++ b/Engine/Plugins/Fab/Source/Fab/Private/FabBrowser.cpp
@@ -693,6 +693,12 @@ bool FFabBrowser::InterceptNavigation(const FString& CurrentUrl, const FString&
 		return false;
 	}
 
+	// Don't remap 3D previewer URLs
+	if (CurrentUrl.Contains(TEXT("/dope/"), ESearchCase::IgnoreCase))
+	{
+		return false;
+	}
+
 	// Remap Fab web URL to plugin URL
 	const FString Remapped = BuildFabPluginUrl(CurrentUrl, BaseUrl);
 	if (!Remapped.Equals(CurrentUrl, ESearchCase::CaseSensitive))
diff --git a/Engine/Plugins/Fab/Source/Fab/Private/Workflows/PackImportWorkflow.cpp b/Engine/Plugins/Fab/Source/Fab/Private/Workflows/PackImportWorkflow.cpp
--- a/Engine/Plugins/Fab/Source/Fab/Private/Workflows/PackImportWorkflow.cpp
+++ b/Engine/Plugins/Fab/Source/Fab/Private/Workflows/PackImportWorkflow.cpp
@@ -314,6 +314,7 @@ void FPackImportWorkflow::CreateDownloadNotification()
 	// Create the notification info
 	FNotificationInfo Info(FText::FromString("Downloading..."));
 
+	TWeakPtr<FFabDownloadRequest> WeakRequest = DownloadRequest;
 	ProgressWidget = SNew(SNotificationProgressWidget)
 		.ProgressText(FText::FromString("Downloading " + AssetName))
 		.HasButton(true)
@@ -321,17 +322,12 @@ void FPackImportWorkflow::CreateDownloadNotification()
 		.ButtonToolTip(FText::FromString("Cancel Pack Import"))
 		.OnButtonClicked(
 			FOnClicked::CreateLambda(
-				[this]()
+				[WeakRequest]()
 				{
 					FAB_LOG("Import Cancelled");
-					if (ProgressWidget.IsValid())
+					if (TSharedPtr<FFabDownloadRequest> PinnedRequest = WeakRequest.Pin())
 					{
-						FAB_LOG("Disabling cancel button");
-						ProgressWidget->DisableCancel();
-					}	
-					if (DownloadRequest)
-					{
-						DownloadRequest->Cancel();
+						PinnedRequest->Cancel();
 					}
 					return FReply::Handled();
 				}
