[Fab] Support for multiple instances

#Synced-CL 52249030

diff --git a/Engine/Plugins/Fab/Source/Fab/Private/FabBrowser.cpp b/Engine/Plugins/Fab/Source/Fab/Private/FabBrowser.cpp
--- a/Engine/Plugins/Fab/Source/Fab/Private/FabBrowser.cpp
+++ b/Engine/Plugins/Fab/Source/Fab/Private/FabBrowser.cpp
@@ -28,6 +28,12 @@
 
 #include "Math/Vector2D.h"
 
+#include "Runtime/Launch/Resources/Version.h"
+
+#if PLATFORM_WINDOWS && (ENGINE_MAJOR_VERSION > 5 || (ENGINE_MAJOR_VERSION == 5 && ENGINE_MINOR_VERSION >= 7))
+#include "Windows/WindowsApplication.h"
+#endif
+
 #include "Misc/FileHelper.h"
 
 #include "Styling/CoreStyle.h"
@@ -46,16 +52,14 @@
 
 #define LOCTEXT_NAMESPACE "Fab"
 
-TSharedPtr<SWebBrowser> FFabBrowser::WebBrowserInstance = nullptr;
-
-TObjectPtr<UFabBrowserApi> FFabBrowser::JavascriptApi = nullptr;
-
-TSharedPtr<SDockTab> FFabBrowser::DockTab                     = nullptr;
 TUniquePtr<FSlateStyleSet> FFabBrowser::SlateStyleSet         = nullptr;
-TSharedPtr<IWebBrowserWindow> FFabBrowser::WebBrowserWindow   = nullptr;
 TObjectPtr<const UFabSettings> FFabBrowser::FabPluginSettings = nullptr;
 
-const FName FFabBrowser::TabId = TEXT("FabTab");
+int32 FFabBrowser::FabTabCounter = 0;
+FName FFabBrowser::ActiveFabTabId = NAME_None;
+TMap<FName, FFabBrowserTabContext> FFabBrowser::TabsById = {};
+TSharedPtr<FTabManager> FFabBrowser::TabManager;
+bool FFabBrowser::bIsShuttingDown = false;
 
 const FText FFabBrowser::FabLabel           = LOCTEXT("Fab.Label", "Fab");
 const FText FFabBrowser::FabTooltip         = LOCTEXT("Fab.Tooltip", "Get content from Fab");
@@ -63,10 +67,36 @@ const FName FFabBrowser::FabMenuIconName    = TEXT("Fab.MenuIcon");
 const FName FFabBrowser::FabAssetIconName   = TEXT("Fab.AssetIcon");
 const FName FFabBrowser::FabToolbarIconName = TEXT("Fab.ToolbarIcon");
 
+static const FName HostTabId(TEXT("FabBrowser.Host"));
+
+// UE 5.7 upgraded CEF, which changed how native-rendered browser windows track
+// their viewport when the parent editor window is moved or resized. This message
+// handler intercepts WM_EXITSIZEMOVE to refresh the viewport and reparent tabs.
+#if PLATFORM_WINDOWS && (ENGINE_MAJOR_VERSION > 5 || (ENGINE_MAJOR_VERSION == 5 && ENGINE_MINOR_VERSION >= 7))
+class FFabWindowMoveMessageHandler final : public IWindowsMessageHandler
+{
+public:
+	TFunction<void()> OnMoveEnded;
+
+	virtual bool ProcessMessage(HWND, uint32 msg, WPARAM, LPARAM, int32&) override
+	{
+		if (msg == WM_EXITSIZEMOVE && OnMoveEnded)
+		{
+			OnMoveEnded();
+			return false;
+		}
+		return false;
+	}
+};
+
+static TSharedPtr<FFabWindowMoveMessageHandler> WindowMoveMessageHandler;
+#endif
+
 void FFabBrowser::Init()
 {
+	bIsShuttingDown = false;
+	PreventCaptchaFailure();
 	RegisterSlateStyle();
-	RegisterNomadTab();
 	SetupEntryPoints();
 	ExtendContextMenuInContentBrowser();
 }
@@ -103,9 +133,9 @@ void FFabBrowser::SetupEntryPoints()
 {
 	const FUIAction InvokeTabAction = FUIAction(
 		FExecuteAction::CreateLambda(
-			[LevelEditorModuleName = FName("LevelEditor")]()
+			[]()
 			{
-				FModuleManager::GetModuleChecked<FLevelEditorModule>(LevelEditorModuleName).GetLevelEditorTabManager()->TryInvokeTab(TabId);
+				CreateNewFabTab();
 			}
 		),
 		FCanExecuteAction()
@@ -141,12 +171,12 @@ void FFabBrowser::SetupEntryPoints()
 		FToolMenuEntry& FabEntry = ContentSection.AddEntry(FabMenuEntry);
 		FabEntry.InsertPosition.Position = EToolMenuInsertType::First;
 	}
-	
+
 	// Add a Fab entry to the Content Browser's Add popup menu
 	UToolMenus::Get()->ExtendMenu(TEXT("ContentBrowser.AddNewContextMenu"))->AddSection(TEXT("ContentBrowserGetContent"), LOCTEXT("GetContentText", "Get Content")).AddEntry(
 		FToolMenuEntry::InitMenuEntry(TEXT("OpenFabWindow"), FabLabel, FabTooltip, FSlateIcon(SlateStyleSet->GetStyleSetName(), FabMenuIconName), InvokeTabAction)
 	);
-	
+
 	{
 		FToolMenuEntry& FabEntry = UToolMenus::Get()->ExtendMenu("LevelEditor.LevelEditorToolBar.AddQuickMenu")->FindOrAddSection("Content").AddEntry(FabMenuEntry);
 		FabEntry.InsertPosition.Name = "ImportContent";
@@ -181,8 +211,8 @@ TSharedRef<FExtender> FFabBrowser::OnExtendContentBrowserAssetSelectionMenu(cons
 			[FabListingId](FMenuBuilder& MenuBuilder)
 			{
 				MenuBuilder.AddMenuEntry(
-					FText::FromString("View in Fab"),
-					FText::FromString("View the asset in Fab plugin"),
+					LOCTEXT("Fab.ViewInFab", "View in Fab"),
+					LOCTEXT("Fab.ViewInFabTooltip", "View the asset in Fab plugin"),
 					FSlateIcon(SlateStyleSet->GetStyleSetName(), FabMenuIconName),
 					FUIAction(
 						FExecuteAction::CreateLambda(
@@ -218,204 +248,561 @@ TSharedRef<SWidget> FFabBrowser::OnFabAssetIconGenerate(const FAssetData& AssetD
 		[
 			SNew(SImage)
 			.Image(FabImage)
-			.ToolTipText(FText::FromString("Imported from FAB"))
+			.ToolTipText(LOCTEXT("Fab.ImportedFromFab", "Imported from FAB"))
 		];
 }
 
-void FFabBrowser::RegisterNomadTab()
+// Helpers
+
+FName FFabBrowser::MakeNextFabTabId()
+{
+	++FabTabCounter;
+	const FString TabName = FString::Printf(TEXT("Fab%d"), FabTabCounter);
+	return FName(*TabName);
+}
+
+bool FFabBrowser::EnsureWebBrowserAvailable()
 {
-	auto RegisterSpawner = [](TSharedPtr<ILevelEditor>)
+	// // Call Get() first — it loads the module if needed. Checking IsAvailable()
+	// // before Get() would fail when the module simply hasn't been loaded yet.
+	const IWebBrowserModule& WebBrowserModule = IWebBrowserModule::Get();
+	if (!IWebBrowserModule::IsAvailable() || !WebBrowserModule.IsWebModuleAvailable())
 	{
-		FGlobalTabmanager::Get()->RegisterNomadTabSpawner(TabId, FOnSpawnTab::CreateStatic(OpenTab)).SetAutoGenerateMenuEntry(false).SetDisplayName(FabLabel).
-		                          SetTooltipTextAttribute(FabTooltip).SetIcon(FSlateIcon(SlateStyleSet->GetStyleSetName(), FabMenuIconName));
-	};
+		FAB_LOG_ERROR("WebBrowserModule isn't available! Can't display Fab tab contents.");
+		return false;
+	}
+
+	return true;
+}
+
+static TSharedRef<SWidget> MakeErrorContent()
+{
+	return SNew(SVerticalBox)
+		+ SVerticalBox::Slot()
+		.AutoHeight()
+		.Padding(FMargin(0, 32, 0, 0))
+		.HAlign(HAlign_Center)
+		[
+			SNew(SImage)
+				.Image(FAppStyle::Get().GetBrush("Icons.Warning"))
+		]
+		+ SVerticalBox::Slot()
+		.AutoHeight()
+		.HAlign(HAlign_Center)
+		[
+			SNew(STextBlock)
+				.Text(LOCTEXT("Fab.BrowserError", "Unable to display the Fab Browser"))
+				.Justification(ETextJustify::Center)
+		];
+}
+
+static void SpawnErrorTab()
+{
+	const FName ErrorTabId(TEXT("FabError"));
+	FGlobalTabmanager::Get()->UnregisterNomadTabSpawner(ErrorTabId);
+	FGlobalTabmanager::Get()->RegisterNomadTabSpawner(
+		ErrorTabId,
+		FOnSpawnTab::CreateLambda([](const FSpawnTabArgs&)
+		{
+			return SNew(SDockTab)
+				.TabRole(ETabRole::NomadTab)
+				[
+					MakeErrorContent()
+				];
+		})
+	)
+	.SetDisplayName(LOCTEXT("Fab.Label", "Fab"))
+	.SetMenuType(ETabSpawnerMenuType::Hidden);
+
+	FGlobalTabmanager::Get()->TryInvokeTab(ErrorTabId);
+}
+
+FString FFabBrowser::BuildIndexFileUrl()
+{
+	const FString PluginPath = IPluginManager::Get().FindPlugin(TEXT("Fab"))->GetBaseDir();
+	const FString IndexPath  = FPaths::ConvertRelativePathToFull(
+		FPaths::Combine(PluginPath, TEXT("ThirdParty"), TEXT("index.html"))
+	);
+
+	// "file:///" + absolute path
+	return FPaths::Combine(TEXT("file:///"), IndexPath);
+}
 
-	FLevelEditorModule& LevelEditorModule = FModuleManager::LoadModuleChecked<FLevelEditorModule>("LevelEditor");
-	if (LevelEditorModule.GetLevelEditorInstance().IsValid())
+// Tab Creation
+
+void FFabBrowser::CreateNewFabTab(const FString& Url)
+{
+	FabAuthentication::LoginUsingPersist();
+
+	if (!FabPluginSettings)
+	{
+		FabPluginSettings = GetDefault<UFabSettings>();
+	}
+
+	if (!EnsureWebBrowserAvailable())
+	{
+		SpawnErrorTab();
+		return;
+	}
+
+#if PLATFORM_MAC
+	// On Mac, WebKit is always natively rendered, so always use native mode
+	CreateNativeFabTab(Url);
+#else
+	// On Windows, respect the user's rendering mode preference
+	if (FabPluginSettings->bEnableNativeRendering)
 	{
-		RegisterSpawner(LevelEditorModule.GetLevelEditorInstance().Pin());
+		CreateNativeFabTab(Url);
 	}
 	else
 	{
-		LevelEditorModule.OnLevelEditorCreated().AddLambda(RegisterSpawner);
+		CreateOsrFabTab(Url);
 	}
+#endif
 }
 
-FString FFabBrowser::GetUrl()
+void FFabBrowser::CreateNativeFabTab(const FString& Url)
 {
-	if (FabPluginSettings == nullptr)
+	// Create window if needed
+	if (!TabManager.IsValid())
 	{
-		return TEXT("https://www.fab.com/plugins/ue5");
+		const TSharedRef<SDockTab> RootDockTab =
+			SNew(SDockTab)
+			.TabRole(MajorTab);
+
+		TabManager = FGlobalTabmanager::Get()->NewTabManager(RootDockTab);
+
+		TabManager->InsertNewDocumentTab(
+			HostTabId,
+			FTabManager::ESearchPreference::PreferLiveTab,
+			RootDockTab
+		);
 	}
 
-	switch (FabPluginSettings->Environment)
+	const FName TabId = MakeNextFabTabId();
+
+	// Create a new major tab
+	const TSharedRef<SDockTab> DockTab =
+		SNew(SDockTab)
+		.TabRole(DocumentTab)
+		.Label(FabLabel)
+		.OnTabClosed_Lambda([TabId](TSharedRef<SDockTab> Tab) { HandleTabClosed(Tab, TabId); });
+
+	// Pre-register the tab context so HandleTabClosed can find and mark it
+	// as closing if the user closes the tab before the deferred attach fires.
+	FFabBrowserTabContext& PreContext = TabsById.FindOrAdd(TabId);
+	PreContext.TabId = TabId;
+	PreContext.DockTab = DockTab;
+
+	TabManager->InsertNewDocumentTab(
+		HostTabId,
+		FTabManager::ESearchPreference::PreferLiveTab,
+		DockTab
+	);
+
+	// Defer browser attach to the next frame so the DockTab has a valid
+	// parent window and correct bounds. Attaching synchronously before the
+	// tab is fully docked causes the native CEF window to briefly fill the
+	// entire screen.
+	// Capture a weak pointer to break the reference cycle (the active timer
+	// is owned by DockTab, so capturing a TSharedRef would prevent destruction).
+	TWeakPtr<SDockTab> WeakDockTab = DockTab;
+	DockTab->RegisterActiveTimer(0.0f,
+		FWidgetActiveTimerDelegate::CreateLambda(
+			[WeakDockTab, TabId, Url](double /*InCurrentTime*/, float /*InDeltaTime*/) -> EActiveTimerReturnType
+			{
+				TSharedPtr<SDockTab> PinnedTab = WeakDockTab.Pin();
+				if (!PinnedTab.IsValid())
+				{
+					return EActiveTimerReturnType::Stop;
+				}
+
+				// If the tab was already closed between insertion and this callback, bail out.
+				if (FFabBrowserTabContext* PreRegistered = TabsById.Find(TabId))
+				{
+					if (PreRegistered->bIsClosing)
+					{
+						return EActiveTimerReturnType::Stop;
+					}
+				}
+
+				AttachBrowserToDockTab(PinnedTab.ToSharedRef(), /*bNative*/ true, TabId, Url);
+				SetupTabWindowMoveListenerForNativeRendering(PinnedTab.ToSharedRef());
+				return EActiveTimerReturnType::Stop;
+			}
+		)
+	);
+}
+
+void FFabBrowser::CreateOsrFabTab(const FString& Url)
+{
+#if (ENGINE_MAJOR_VERSION == 5 && ENGINE_MINOR_VERSION <= 6)
+	// Limit the number of OSR tabs that can be opened
+	// due to the problem of browser freezing on earlier
+	// engine versions (5.3-5.6)
+
+	// Count the current number of OSR tabs
+	int32 Count = 0;
+	for (const TPair<FName, FFabBrowserTabContext>& Tab : TabsById)
 	{
-		default: // fall through
-		case EFabEnvironment::Prod:  // fall through
-		case EFabEnvironment::Gamedev: // fall through
-		case EFabEnvironment::Test:  // fall through
+		if (!Tab.Value.bIsNativelyRendered)
 		{
-			FString Url = FabPluginSettings->GetUrlFromEnvironment();
-			Url += TEXT("/plugins/ue5");
-			return Url;
+			Count += 1;
 		}
-		case EFabEnvironment::CustomUrl:
+	}
+
+	if (Count > 3)
+	{
+		CreateNativeFabTab(Url);
+		return;
+	}
+#endif
+
+	const FName TabId = MakeNextFabTabId();
+	FGlobalTabmanager::Get()->RegisterNomadTabSpawner(
+		TabId,
+		FOnSpawnTab::CreateLambda([TabId, Url](const FSpawnTabArgs&)
 		{
-			return FabPluginSettings->CustomUrl;
+			const TSharedRef<SDockTab> DockTab =
+				SNew(SDockTab)
+				.TabRole(ETabRole::NomadTab)
+				.OnTabClosed_Lambda([TabId](TSharedRef<SDockTab> Tab) { HandleTabClosed(Tab, TabId); })
+				.OnTabActivated_Static(&FFabBrowser::HandleTabActivated, TabId);
+
+			AttachBrowserToDockTab(DockTab, /*bNative*/ false, TabId, Url);
+			return DockTab;
+		})
+	)
+	.SetDisplayName(FText::FromName(TabId))
+	.SetMenuType(ETabSpawnerMenuType::Hidden);
+
+	FGlobalTabmanager::Get()->TryInvokeTab(TabId);
+}
+
+void FFabBrowser::PreventWebBrowserFreezeInOsrMode(const FName& TabId, const FFabBrowserTabContext& Context)
+{
+	if (!Context.bIsNativelyRendered)
+	{
+		Context.WebBrowserInstance->SetOnMouseEnter(FNoReplyPointerEventHandler::CreateLambda([TabId](const FGeometry& Geometry, const FPointerEvent& PointerEvent)
+		{
+			// Disable all other web browser instances
+			for (const TTuple<FName, FFabBrowserTabContext>& Element : TabsById)
+			{
+				// Don't disable the active tab
+				if (ActiveFabTabId != Element.Key && !Element.Value.bIsNativelyRendered && Element.Value.WebBrowserWindow.IsValid())
+				{
+					Element.Value.WebBrowserWindow->SetIsDisabled(true);
+				}
+			}
+
+			// Enable the one in focus
+			if (FFabBrowserTabContext* const& FabBrowserTabContext = TabsById.Find(TabId))
+			{
+				if (FabBrowserTabContext->WebBrowserWindow.IsValid())
+				{
+					FabBrowserTabContext->WebBrowserWindow->SetIsDisabled(false);
+				}
+			}
+		}));
+	}
+}
+
+void FFabBrowser::AttachBrowserToDockTab(const TSharedRef<SDockTab>& DockTab, const bool bNativeRendering, const FName TabId, const FString& Url)
+{
+	FFabBrowserTabContext& Context = TabsById.FindOrAdd(TabId);
+	Context.TabId = TabId;
+	Context.DockTab = DockTab;
+
+	// Create JS API — each tab should get exactly one bridge object.
+	if (Context.JavascriptApi)
+	{
+		ensureMsgf(false, TEXT("AttachBrowserToDockTab called with an existing JavascriptApi on TabId %s; cleaning up."), *TabId.ToString());
+		Context.JavascriptApi->RemoveFromRoot();
+		Context.JavascriptApi = nullptr;
+	}
+	Context.JavascriptApi = NewObject<UFabBrowserApi>();
+	Context.JavascriptApi->AddToRoot();
+
+	// Browser settings
+	FCreateBrowserWindowSettings Settings;
+	Settings.InitialURL = Url.IsEmpty() ? BuildIndexFileUrl() : Url;
+	Settings.BrowserFrameRate = 60;
+#if PLATFORM_MAC
+	Settings.bMobileJSReturnInDict = false;
+#endif
+
+	Context.bIsNativelyRendered = bNativeRendering;
+
+	if (bNativeRendering)
+	{
+		if (DockTab->GetParentWindow().IsValid())
+		{
+			Settings.OSWindowHandle = DockTab->GetParentWindow()->GetNativeWindow()->GetOSWindowHandle();
 		}
+
+		SetupDockTabListenersForNativeRendering(DockTab, TabId);
+	}
+
+	IWebBrowserSingleton* Singleton = IWebBrowserModule::Get().GetSingleton();
+	Singleton->SetDevToolsShortcutEnabled(true);
+	Context.WebBrowserWindow = Singleton->CreateBrowserWindow(Settings);
+
+	if (!Context.WebBrowserWindow.IsValid())
+	{
+		FAB_LOG_ERROR("WebBrowserWindow is invalid! Can't display Fab tab contents.");
+		Context.JavascriptApi->RemoveFromRoot();
+		Context.JavascriptApi = nullptr;
+		DockTab->SetContent(MakeErrorContent());
+		return;
 	}
+
+	PreventBrowserKeyConflictWithEditor(Context);
+	InstallUrlRestrictions(Context);
+	InstallDebugPopup(Context);
+
+	const bool bShowAddressBar = (FabPluginSettings->Environment == EFabEnvironment::CustomUrl);
+	SAssignNew(Context.WebBrowserInstance, SWebBrowser, Context.WebBrowserWindow)
+	.ShowAddressBar(bShowAddressBar)
+	.ShowControls(bShowAddressBar)
+	.OnTitleChanged_Lambda([TabId](const FText& Title)
+	{
+		if (FFabBrowserTabContext* const& FabBrowserTabContext = TabsById.Find(TabId))
+		{
+			if (const TSharedPtr<SDockTab>& Tab = FabBrowserTabContext->DockTab.Pin())
+			{
+				if (!Title.IsEmpty())
+				{
+					Tab->SetLabel(Title);
+				}
+			}
+		}
+	});
+
+	PreventWebBrowserFreezeInOsrMode(TabId, Context);
+
+	DockTab->SetContent(Context.WebBrowserInstance.ToSharedRef());
+
+	Context.WebBrowserInstance->BindUObject(TEXT("fab"), Cast<UObject>(Context.JavascriptApi), true);
+	Context.WebBrowserWindow->Reload();
+
+	ActiveFabTabId = TabId;
+
+#if PLATFORM_MAC && !WITH_CEF3
+	FSlateApplication::Get().ToggleDisableLastDragOnDragEnter(true);
+#endif
 }
 
-const ISlateStyle& FFabBrowser::GetStyleSet()
+void FFabBrowser::SetupDockTabListenersForNativeRendering(const TSharedRef<SDockTab>& DockTab, const FName TabId)
 {
-	return *(SlateStyleSet.Get());
+	DockTab->SetOnTabDraggedOverDockArea(FSimpleDelegate::CreateLambda([TabId]
+	{
+		// Disable the web browser that is being dragged
+		if (FFabBrowserTabContext* const& TabContext = TabsById.Find(TabId))
+		{
+			if (!TabContext->bIsClosing && TabContext->WebBrowserWindow.IsValid())
+			{
+				TabContext->WebBrowserWindow->SetIsDisabled(true);
+			}
+		}
+	}));
+
+	TWeakPtr<SDockTab> WeakRelocatedTab = DockTab;
+	DockTab->SetOnTabRelocated(FSimpleDelegate::CreateLambda([WeakRelocatedTab]
+	{
+		if (WeakRelocatedTab.Pin().IsValid())
+		{
+			RefreshNativeTabViewPort();
+			ReparentDockTabs();
+		}
+	}));
+
+	DockTab->SetOnTabActivated(SDockTab::FOnTabActivatedCallback::CreateLambda([](TSharedRef<SDockTab>, ETabActivationCause)
+	{
+		ReparentDockTabs();
+	}));
 }
 
-TSharedRef<SDockTab> FFabBrowser::OpenTab(const FSpawnTabArgs&)
+void FFabBrowser::SetupTabWindowMoveListenerForNativeRendering(const TSharedRef<SDockTab>& DockTab)
 {
-	FabAuthentication::LoginUsingPersist();
-	
-	FabPluginSettings = GetDefault<UFabSettings>();
+	const TSharedPtr<SWindow> ParentWindow = DockTab->GetParentWindow();
+	if (!ParentWindow.IsValid())
+	{
+		return;
+	}
 
-	JavascriptApi = NewObject<UFabBrowserApi>();
-	JavascriptApi->AddToRoot(); //Don't garbage collect
+#if PLATFORM_WINDOWS && (ENGINE_MAJOR_VERSION > 5 || (ENGINE_MAJOR_VERSION == 5 && ENGINE_MINOR_VERSION >= 7))
+	FWindowsApplication* WindowsApp = static_cast<FWindowsApplication*>(FSlateApplication::Get().GetPlatformApplication().Get());
+	if (!WindowsApp)
+	{
+		return;
+	}
 
-	IWebBrowserModule& WebBrowserModule = IWebBrowserModule::Get();
-	if (!IWebBrowserModule::IsAvailable() || !WebBrowserModule.IsWebModuleAvailable())
+	if (!WindowMoveMessageHandler.IsValid())
 	{
-		FAB_LOG_ERROR("WebBrowserModule isn't available! Can't display Fab tab contents.");
-		return SNew(SDockTab)
-			.TabRole(ETabRole::NomadTab)
-			[
-				SNew(SVerticalBox)
-					+ SVerticalBox::Slot()
-					.AutoHeight()
-					.Padding(FMargin(0, 32, 0, 0))
-					.HAlign(HAlign_Center)
-					[
-						SNew(SImage)
-							.Image(FAppStyle::Get().GetBrush("Icons.Warning"))
-					]
-					+ SVerticalBox::Slot()
-					.AutoHeight()
-					.HAlign(HAlign_Center)
-					[
-						SNew(STextBlock)
-							.Text(FText::FromString("Unable to display the Fab Browser"))
-							.Justification(ETextJustify::Center)
-					]
-			];
+		WindowMoveMessageHandler = MakeShared<FFabWindowMoveMessageHandler>();
+		WindowMoveMessageHandler->OnMoveEnded = []()
+		{
+			RefreshNativeTabViewPort();
+			ReparentDockTabs();
+		};
+		WindowsApp->AddMessageHandler(*WindowMoveMessageHandler);
 	}
+#endif
+}
 
-	PreventCaptchaFailure();
+void FFabBrowser::ReparentDockTabs()
+{
+	for (const TTuple<FName, FFabBrowserTabContext>& Pair : TabsById)
+	{
+		const FFabBrowserTabContext& TabContext = Pair.Value;
+		if (!TabContext.bIsNativelyRendered || !TabContext.WebBrowserInstance.IsValid() || !TabContext.WebBrowserWindow.IsValid())
+		{
+			continue;
+		}
 
-	FCreateBrowserWindowSettings WindowSettings;
+		TabContext.WebBrowserWindow->SetIsDisabled(false);
 
-	FString PluginPath              = IPluginManager::Get().FindPlugin(TEXT("Fab"))->GetBaseDir();
-	FString IndexUrl                = FPaths::ConvertRelativePathToFull(FPaths::Combine(PluginPath, TEXT("ThirdParty"), TEXT("index.html")));
-	FString FinalUrl                = FPaths::Combine(TEXT("file:///"), IndexUrl);
+		const TSharedPtr<SDockTab>& Tab = TabContext.DockTab.Pin();
+		if (Tab.IsValid() && Tab->GetParentWindow().IsValid())
+		{
+			TabContext.WebBrowserWindow->SetParentWindow(Tab->GetParentWindow());
+		}
+	}
+}
 
-	WindowSettings.InitialURL       = FinalUrl;
+void FFabBrowser::RefreshNativeTabViewPort()
+{
+	for (const TTuple<FName, FFabBrowserTabContext>& Pair : TabsById)
+	{
+		const FFabBrowserTabContext& TabContext = Pair.Value;
+		if (!TabContext.bIsNativelyRendered || !TabContext.WebBrowserInstance.IsValid() || !TabContext.WebBrowserWindow.IsValid())
+		{
+			continue;
+		}
 
-	WindowSettings.BrowserFrameRate = 60;
-#if PLATFORM_MAC
-	WindowSettings.bMobileJSReturnInDict = false;
-#endif
+		TabContext.WebBrowserWindow->SetIsDisabled(false);
 
-	IWebBrowserSingleton* WebBrowserSingleton = WebBrowserModule.GetSingleton();
-	WebBrowserSingleton->SetDevToolsShortcutEnabled(true);
+		const TSharedPtr<SDockTab>& DockTab = TabContext.DockTab.Pin();
+		if (!DockTab.IsValid())
+		{
+			continue;
+		}
+
+		if (const TSharedPtr<SWindow>& ParentWindow = DockTab->GetParentWindow(); ParentWindow.IsValid())
+		{
+			const FVector2D WindowSize = ParentWindow->GetViewportSize();
+			const FGeometry& CachedGeometry = TabContext.WebBrowserInstance->GetCachedGeometry();
+			const UE::Math::TVector2<double>& BrowserPosition = CachedGeometry.GetAbsolutePosition();
+
+			TabContext.WebBrowserWindow->SetViewportSize(
+				FIntPoint(FMath::RoundToInt(WindowSize.X), FMath::RoundToInt(WindowSize.Y)),
+				FIntPoint(FMath::RoundToInt(BrowserPosition.X), FMath::RoundToInt(BrowserPosition.Y))
+			);
+		}
+	}
+}
 
-	WebBrowserWindow = WebBrowserSingleton->CreateBrowserWindow(WindowSettings);
-	if (!WebBrowserWindow.IsValid())
+void FFabBrowser::PreventBrowserKeyConflictWithEditor(const FFabBrowserTabContext& Context)
+{
+	Context.WebBrowserWindow->OnUnhandledKeyUp().BindLambda([](const FKeyEvent&) { return true; });
+	Context.WebBrowserWindow->OnUnhandledKeyDown().BindLambda([](const FKeyEvent&) { return true; });
+}
+
+void FFabBrowser::InstallUrlRestrictions(const FFabBrowserTabContext& Context)
+{
+	if (!FabPluginSettings || FabPluginSettings->Environment == EFabEnvironment::CustomUrl)
 	{
-		FAB_LOG_ERROR("WebBrowserWindow is invalid! Can't display Fab tab contents.");
-		return SNew(SDockTab)
-			.TabRole(ETabRole::NomadTab)
-			[
-				SNew(SVerticalBox)
-					+ SVerticalBox::Slot()
-					.AutoHeight()
-					.Padding(FMargin(0, 32, 0, 0))
-					.HAlign(HAlign_Center)
-					[
-						SNew(SImage)
-							.Image(FAppStyle::Get().GetBrush("Icons.Warning"))
-					]
-					+ SVerticalBox::Slot()
-					.AutoHeight()
-					.HAlign(HAlign_Center)
-					[
-						SNew(STextBlock)
-							.Text(FText::FromString("Unable to display the Fab Browser"))
-							.Justification(ETextJustify::Center)
-					]
-			];
+		return;
 	}
 
-	WebBrowserWindow->OnUnhandledKeyUp().BindLambda([](const FKeyEvent&) { return true; });
-	WebBrowserWindow->OnUnhandledKeyDown().BindLambda([](const FKeyEvent&) { return true; });
-	WebBrowserWindow->OnUrlChanged().AddLambda([](const FString& Url)
+	TWeakPtr<IWebBrowserWindow> WeakBrowserWindow = Context.WebBrowserWindow;
+
+	Context.WebBrowserWindow->OnUrlChanged().AddLambda([WeakBrowserWindow](const FString& Url)
 	{
+		TSharedPtr<IWebBrowserWindow> BrowserWindow = WeakBrowserWindow.Pin();
+		if (!BrowserWindow.IsValid())
+		{
+			return;
+		}
+
 		if (!FabPluginSettings || FabPluginSettings->Environment == EFabEnvironment::CustomUrl)
 		{
 			return;
 		}
 
-		InterceptNavigation(Url, FabPluginSettings->GetUrlFromEnvironment(), GetUrl());
+		InterceptNavigation(BrowserWindow, Url, FabPluginSettings->GetUrlFromEnvironment(), GetUrl());
 	});
 
-	WebBrowserWindow->OnBeforeBrowse().BindLambda([](const FString& Url, const FWebNavigationRequest& Request) -> bool
+	Context.WebBrowserWindow->OnBeforeBrowse().BindLambda([WeakBrowserWindow](const FString& Url, const FWebNavigationRequest& Request) -> bool
 	{
+		TSharedPtr<IWebBrowserWindow> BrowserWindow = WeakBrowserWindow.Pin();
+		if (!BrowserWindow.IsValid())
+		{
+			return false;
+		}
+
 		if (!FabPluginSettings || FabPluginSettings->Environment == EFabEnvironment::CustomUrl)
 		{
-			return false; // Allow navigation
+			return false;
 		}
 
-		return InterceptNavigation(Url, FabPluginSettings->GetUrlFromEnvironment(), GetUrl());
+		return InterceptNavigation(BrowserWindow, Url, FabPluginSettings->GetUrlFromEnvironment(), GetUrl());
 	});
+}
 
-	if (FabPluginSettings->bEnableDebugOptions)
+void FFabBrowser::InstallDebugPopup(const FFabBrowserTabContext& Context)
+{
+	if (!FabPluginSettings || !FabPluginSettings->bEnableDebugOptions)
 	{
-		WebBrowserWindow.Get()->OnCreateWindow().BindLambda(
-			[](const TWeakPtr<IWebBrowserWindow>& NewBrowserWindow, const TWeakPtr<IWebBrowserPopupFeatures>& PopupFeatures)
-			{
-				const TSharedRef<SWindow> DialogMainWindow = SNew(SWindow)
-					.ClientSize(FVector2D(700, 700))
-					.SupportsMaximize(true)
-					.SupportsMinimize(true)
-					[
-						SNew(SVerticalBox)
-						+ SVerticalBox::Slot().HAlign(HAlign_Fill).VAlign(VAlign_Fill)
-						[
-							SNew(SWebBrowser, NewBrowserWindow.Pin())
-						]
-					];
-				FSlateApplication::Get().AddWindow(DialogMainWindow);
-				return true;
-			}
-		);
+		return;
 	}
 
-	bool bShowAddressBar = FabPluginSettings->Environment == EFabEnvironment::CustomUrl;
+	Context.WebBrowserWindow->OnCreateWindow().BindLambda(
+		[](const TWeakPtr<IWebBrowserWindow>& NewBrowserWindow,
+		   const TWeakPtr<IWebBrowserPopupFeatures>& /*PopupFeatures*/)
+		{
+			const TSharedRef<SWindow> Dialog = SNew(SWindow)
+				.ClientSize(FVector2D(700, 700))
+				.SupportsMaximize(true)
+				.SupportsMinimize(true)
+				[
+					SNew(SVerticalBox)
+					+ SVerticalBox::Slot().HAlign(HAlign_Fill).VAlign(VAlign_Fill)
+					[
+						SNew(SWebBrowser, NewBrowserWindow.Pin())
+					]
+				];
 
-	SAssignNew(WebBrowserInstance, SWebBrowser, WebBrowserWindow).ShowAddressBar(bShowAddressBar).ShowControls(bShowAddressBar);
+			FSlateApplication::Get().AddWindow(Dialog);
+			return true;
+		}
+	);
+}
 
-	WebBrowserInstance->BindUObject(TEXT("fab"), Cast<UObject>(JavascriptApi), true);
-	WebBrowserWindow->Reload();
+void FFabBrowser::HandleTabActivated(TSharedRef<SDockTab> ActivatedTab, ETabActivationCause Cause, FName TabId)
+{
+	if (ActiveFabTabId == TabId)
+	{
+		return;
+	}
 
-	SAssignNew(DockTab, SDockTab).TabRole(ETabRole::NomadTab).OnTabClosed_Static(&FFabBrowser::OnPluginTabClosed)
-	[
-		WebBrowserInstance.ToSharedRef()
-	];
-	
-#if PLATFORM_MAC && !WITH_CEF3
-	FSlateApplication::Get().ToggleDisableLastDragOnDragEnter(true);
-#endif
-	
-	return DockTab.ToSharedRef();
+	if (FFabBrowserTabContext* const& FabBrowserTabContext = TabsById.Find(TabId))
+	{
+		if (FabBrowserTabContext->bIsClosing)
+		{
+			return;
+		}
+
+		if (FabBrowserTabContext->WebBrowserWindow.IsValid())
+		{
+			FabBrowserTabContext->WebBrowserWindow->SetIsDisabled(false);
+		}
+	}
+
+	ActiveFabTabId = TabId;
 }
 
-void FFabBrowser::OnPluginTabClosed(TSharedRef<SDockTab> InParentTab)
+void FFabBrowser::HandleTabClosed(TSharedRef<SDockTab> ClosedTab, FName TabId)
 {
 	LogEvent(
 		{
@@ -431,59 +818,244 @@ void FFabBrowser::OnPluginTabClosed(TSharedRef<SDockTab> InParentTab)
 		}
 	);
 
-#if PLATFORM_MAC && !WITH_CEF3
-	if(WebBrowserWindow.IsValid())
+	FFabBrowserTabContext* Context = TabsById.Find(TabId);
+	if (!Context)
 	{
-		WebBrowserWindow->SetIsDisabled(true);
+		return;
 	}
+
+	// Mark as closing immediately to guard other callbacks
+	Context->bIsClosing = true;
+
+	// Unbind JS bridge and clear browser callbacks synchronously
+	// to prevent stale events during the 1.5s delay
+	if (Context->WebBrowserInstance)
+	{
+		Context->WebBrowserInstance->UnbindUObject(TEXT("fab"), Cast<UObject>(Context->JavascriptApi), true);
+	}
+
+	if (Context->WebBrowserWindow.IsValid())
+	{
+		Context->WebBrowserWindow->OnUrlChanged().Clear();
+		Context->WebBrowserWindow->OnBeforeBrowse().Unbind();
+		Context->WebBrowserWindow->OnUnhandledKeyUp().Unbind();
+		Context->WebBrowserWindow->OnUnhandledKeyDown().Unbind();
+		Context->WebBrowserWindow->OnCreateWindow().Unbind();
+	}
+
+	Context->DockTab.Reset();
+
+	// Clear active if needed
+	if (ActiveFabTabId == TabId)
+	{
+		ActiveFabTabId = NAME_None;
+	}
+
+#if PLATFORM_MAC && !WITH_CEF3
+	FSlateApplication::Get().ToggleDisableLastDragOnDragEnter(false);
 #endif
-	
-	AsyncTask(ENamedThreads::AnyBackgroundThreadNormalTask, []()
+
+	// If we are shutting down, do cleanup immediately instead of deferring.
+	// The deferred lambda would fire after Shutdown() has already emptied TabsById,
+	// causing access to stale static data.
+	if (bIsShuttingDown)
+	{
+		if (Context->JavascriptApi)
+		{
+			Context->JavascriptApi->RemoveFromRoot();
+			Context->JavascriptApi = nullptr;
+		}
+
+		Context->WebBrowserInstance.Reset();
+		Context->WebBrowserWindow.Reset();
+
+		FGlobalTabmanager::Get()->UnregisterNomadTabSpawner(TabId);
+		TabsById.Remove(TabId);
+		return;
+	}
+
+	// Defer final resource release so the frontend can register the close event
+	AsyncTask(ENamedThreads::AnyBackgroundThreadNormalTask, [TabId]()
 	{
-		// add a small delay so that the frontend can register the close event
 		FPlatformProcess::Sleep(1.5f);
 
-		AsyncTask(ENamedThreads::GameThread, []()
+		AsyncTask(ENamedThreads::GameThread, [TabId]()
 		{
-			WebBrowserInstance->UnbindUObject(TEXT("fab"), Cast<UObject>(JavascriptApi), true);
-			WebBrowserInstance.Reset();
-			WebBrowserWindow.Reset();
-			DockTab.Reset();
-			
-#if PLATFORM_MAC && !WITH_CEF3
-			FSlateApplication::Get().ToggleDisableLastDragOnDragEnter(false);
-#endif
+			// Guard against Shutdown() having run while we were sleeping
+			if (bIsShuttingDown)
+			{
+				return;
+			}
+
+			FFabBrowserTabContext* DeferredContext = TabsById.Find(TabId);
+			if (!DeferredContext)
+			{
+				return;
+			}
+
+			if (DeferredContext->JavascriptApi)
+			{
+				DeferredContext->JavascriptApi->RemoveFromRoot();
+				DeferredContext->JavascriptApi = nullptr;
+			}
+
+			DeferredContext->WebBrowserInstance.Reset();
+			DeferredContext->WebBrowserWindow.Reset();
+
+			FGlobalTabmanager::Get()->UnregisterNomadTabSpawner(TabId);
+			TabsById.Remove(TabId);
 		});
 	});
 }
 
-void FFabBrowser::ExecuteJavascript(const FString& InSrcScript)
+FFabBrowserTabContext* FFabBrowser::GetActiveDockTabContext()
+{
+	return TabsById.Find(ActiveFabTabId);
+}
+
+TObjectPtr<UFabBrowserApi> FFabBrowser::GetBrowserApi()
 {
-	if (WebBrowserInstance == nullptr)
+	FFabBrowserTabContext* FabBrowserTabContext = TabsById.Find(ActiveFabTabId);
+	if (!FabBrowserTabContext)
 	{
-		return;
+		return nullptr;
+	}
+
+	return FabBrowserTabContext->JavascriptApi;
+}
+
+FString FFabBrowser::GetUrl()
+{
+	if (FabPluginSettings == nullptr)
+	{
+		return TEXT("https://www.fab.com/plugins/ue5");
+	}
+
+	switch (FabPluginSettings->Environment)
+	{
+		default: // fall through
+		case EFabEnvironment::Prod:  // fall through
+		case EFabEnvironment::Gamedev: // fall through
+		case EFabEnvironment::Test:  // fall through
+		{
+			FString Url = FabPluginSettings->GetUrlFromEnvironment();
+			Url += TEXT("/plugins/ue5");
+			return Url;
+		}
+		case EFabEnvironment::CustomUrl:
+		{
+			return FabPluginSettings->CustomUrl;
+		}
 	}
+}
 
-	WebBrowserInstance->ExecuteJavascript(InSrcScript);
+const ISlateStyle& FFabBrowser::GetStyleSet()
+{
+	return *(SlateStyleSet.Get());
 }
 
 void FFabBrowser::Shutdown()
 {
-	WebBrowserInstance.Reset();
-	WebBrowserWindow.Reset();
-	DockTab.Reset();
+	bIsShuttingDown = true;
+
+#if PLATFORM_WINDOWS && (ENGINE_MAJOR_VERSION > 5 || (ENGINE_MAJOR_VERSION == 5 && ENGINE_MINOR_VERSION >= 7))
+	if (WindowMoveMessageHandler.IsValid())
+	{
+		if (FSlateApplication::IsInitialized())
+		{
+			if (FWindowsApplication* WindowsApp = static_cast<FWindowsApplication*>(FSlateApplication::Get().GetPlatformApplication().Get()))
+			{
+				WindowsApp->RemoveMessageHandler(*WindowMoveMessageHandler);
+			}
+		}
+		WindowMoveMessageHandler.Reset();
+	}
+#endif
+
 	FSlateStyleRegistry::UnRegisterSlateStyle(*SlateStyleSet);
-	FGlobalTabmanager::Get()->UnregisterNomadTabSpawner(TabId);
+
+	// Close / cleanup any still-open tabs and remove spawners
+	// Copy keys because HandleTabClosed/Remove will mutate TabsById.
+	TArray<FName> TabIds;
+	TabsById.GetKeys(TabIds);
+
+	for (const FName& TabId : TabIds)
+	{
+		// If tab is still open, request it to close.
+		if (FFabBrowserTabContext* Context = TabsById.Find(TabId))
+		{
+			if (TSharedPtr<SDockTab> Tab = Context->DockTab.Pin())
+			{
+				// This will trigger OnTabClosed -> HandleTabClosed (in normal cases)
+				Tab->RequestCloseTab();
+			}
+			else
+			{
+				// Tab slate object already gone; do direct cleanup
+				if (Context->JavascriptApi)
+				{
+					Context->JavascriptApi->RemoveFromRoot();
+					Context->JavascriptApi = nullptr;
+				}
+				Context->WebBrowserInstance.Reset();
+				Context->WebBrowserWindow.Reset();
+			}
+		}
+
+		// Always unregister the spawner (safe even if already unregistered)
+		FGlobalTabmanager::Get()->UnregisterNomadTabSpawner(TabId);
+	}
+
+	// In case some OnTabClosed didn't fire (shutdown ordering), force-clear leftovers.
+	for (auto& Pair : TabsById)
+	{
+		FFabBrowserTabContext& Context = Pair.Value;
+
+		if (Context.JavascriptApi)
+		{
+			Context.JavascriptApi->RemoveFromRoot();
+			Context.JavascriptApi = nullptr;
+		}
+
+		Context.WebBrowserInstance.Reset();
+		Context.WebBrowserWindow.Reset();
+		Context.DockTab.Reset();
+	}
+
+	// Clean up error tab spawner in case SpawnErrorTab was called
+	FGlobalTabmanager::Get()->UnregisterNomadTabSpawner(FName(TEXT("FabError")));
+
+	TabsById.Empty();
+	ActiveFabTabId = NAME_None;
 }
 
-void FFabBrowser::LoggedIn(const FString& InAccessToken)
+void FFabBrowser::NotifyAllTabsLogout()
 {
-	ExecuteJavascript(FString::Printf(TEXT("window.ue.fab.onLoginSuccessful('%s')"), *InAccessToken));
+	for (const TPair<FName, FFabBrowserTabContext>& Pair : TabsById)
+	{
+		UFabBrowserApi* Api = Pair.Value.JavascriptApi;
+		if (IsValid(Api))
+		{
+			Api->ExecuteOnLogoutCallback();
+		}
+	}
 }
 
-void FFabBrowser::GetSignedUrl(const FString& AssetId, const int32 Tier)
+void FFabBrowser::NotifyAllTabsLogin(const FString& AccessToken)
 {
-	ExecuteJavascript(FString::Printf(TEXT("window.ue.fab.getSignedUrl('%s', %d)"), *AssetId, Tier));
+	for (const TPair<FName, FFabBrowserTabContext>& Pair : TabsById)
+	{
+		UFabBrowserApi* Api = Pair.Value.JavascriptApi;
+		if (IsValid(Api))
+		{
+			Api->ExecuteOnLoginCallback(AccessToken);
+		}
+	}
+}
+
+void FFabBrowser::LoggedIn(const FString& InAccessToken)
+{
+	NotifyAllTabsLogin(InAccessToken);
 }
 
 void FFabBrowser::LogEvent(const FFabAnalyticsPayload& Payload)
@@ -491,13 +1063,23 @@ void FFabBrowser::LogEvent(const FFabAnalyticsPayload& Payload)
 	FString JSONPayload;
 	FJsonObjectConverter::UStructToJsonObjectString(Payload, JSONPayload, 0, 0, 0, nullptr, false);
 
-	ExecuteJavascript(FString::Printf(TEXT("window.ue.fab.logevent('%s')"), *JSONPayload));
+	const FFabBrowserTabContext* FabBrowserTabContext = TabsById.Find(ActiveFabTabId);
+	if (!FabBrowserTabContext)
+	{
+		return;
+	}
+
+	UFabBrowserApi* Api = FabBrowserTabContext->JavascriptApi;
+	if (IsValid(Api))
+	{
+		Api->ExecuteOnLogEventCallback(JSONPayload);
+	}
 }
 
 void FFabBrowser::ShowSettings()
 {
-	AsyncTask(ENamedThreads::Type::GameThread, []() 
-    {
+	AsyncTask(ENamedThreads::Type::GameThread, []()
+	{
 		const TSharedRef<SWindow> Window = SNew(SWindow)
 			.Title(LOCTEXT("FabSettingsLabel", "Fab Settings"))
 			.ClientSize(FVector2D(600.f, 300.f))
@@ -518,13 +1100,29 @@ void FFabBrowser::ShowSettings()
 
 void FFabBrowser::OpenURL(const FString& InURL)
 {
-	FModuleManager::GetModuleChecked<FLevelEditorModule>("LevelEditor").GetLevelEditorTabManager()->TryInvokeTab(TabId);
-	if (WebBrowserWindow.Get()->GetUrl() != InURL)
+	if (ActiveFabTabId != NAME_None)
 	{
-		WebBrowserWindow.Get()->LoadURL(InURL);
+		if (FFabBrowserTabContext* const& FabBrowserTabContext = TabsById.Find(ActiveFabTabId))
+		{
+			if (!FabBrowserTabContext->WebBrowserWindow.IsValid())
+			{
+				UE_LOG(LogFab, Warning, TEXT("OpenURL: Active tab '%s' has no valid WebBrowserWindow (browser may still be attaching). URL not loaded: %s"), *ActiveFabTabId.ToString(), *InURL);
+				return;
+			}
+			FModuleManager::GetModuleChecked<FLevelEditorModule>("LevelEditor").GetLevelEditorTabManager()->TryInvokeTab(FabBrowserTabContext->TabId);
+			if (FabBrowserTabContext->WebBrowserWindow->GetUrl() != InURL)
+			{
+				FabBrowserTabContext->WebBrowserWindow->LoadURL(InURL);
+			}
+		}
 	}
 }
 
+void FFabBrowser::OpenInNewTab(const FString& Url)
+{
+	CreateNewFabTab(Url);
+}
+
 static bool ExtractHostFromUrl(const FString& Url, FString& OutHost)
 {
 	FString Working = Url;
@@ -652,9 +1250,9 @@ static bool BlockAndOpenExternally(const FString& Url, const TSharedPtr<IWebBrow
 	return true; // Cancel original navigation
 }
 
-bool FFabBrowser::InterceptNavigation(const FString& CurrentUrl, const FString& FabBaseUrl, const FString& FallbackUrl)
+bool FFabBrowser::InterceptNavigation(const TSharedPtr<IWebBrowserWindow>& BrowserWindow, const FString& CurrentUrl, const FString& FabBaseUrl, const FString& FallbackUrl)
 {
-	if (!WebBrowserWindow.IsValid())
+	if (!BrowserWindow.IsValid())
 	{
 		return false;
 	}
@@ -684,7 +1282,7 @@ bool FFabBrowser::InterceptNavigation(const FString& CurrentUrl, const FString&
 	// Not a Fab URL, open externally
 	if (!bIsFabUrl)
 	{
-		return BlockAndOpenExternally(CurrentUrl, WebBrowserWindow, FallbackUrl);
+		return BlockAndOpenExternally(CurrentUrl, BrowserWindow, FallbackUrl);
 	}
 
 	// Already a plugin Fab URL, can ignore
@@ -704,7 +1302,7 @@ bool FFabBrowser::InterceptNavigation(const FString& CurrentUrl, const FString&
 	if (!Remapped.Equals(CurrentUrl, ESearchCase::CaseSensitive))
 	{
 		FAB_LOG("Loading remapped URL [%s]", *Remapped);
-		WebBrowserWindow->LoadURL(Remapped);
+		BrowserWindow->LoadURL(Remapped);
 		return true; // cancel original navigation
 	}
 
diff --git a/Engine/Plugins/Fab/Source/Fab/Private/FabBrowser.h b/Engine/Plugins/Fab/Source/Fab/Private/FabBrowser.h
--- a/Engine/Plugins/Fab/Source/Fab/Private/FabBrowser.h
+++ b/Engine/Plugins/Fab/Source/Fab/Private/FabBrowser.h
@@ -23,7 +23,7 @@ struct FFabAnalyticsEventValue
 
 	UPROPERTY()
 	FString Platform;
-	
+
 	UPROPERTY()
 	FFabApiVersion ApiVersion;
 };
@@ -52,17 +52,35 @@ struct FFabAnalyticsPayload
 	FFabAnalyticsEventValue EventValue;
 };
 
+struct FFabBrowserTabContext
+{
+	FName TabId = NAME_None;
+
+	// Slate/CEF objects
+	TWeakPtr<SDockTab> DockTab;
+	TSharedPtr<IWebBrowserWindow> WebBrowserWindow;
+	TSharedPtr<SWebBrowser> WebBrowserInstance;
+
+	// JS bridge
+	TObjectPtr<UFabBrowserApi> JavascriptApi = nullptr;
+
+	bool bIsNativelyRendered = false;
+	bool bIsClosing = false;
+};
+
 class FFabBrowser
 {
 private:
-	static TSharedPtr<SWebBrowser> WebBrowserInstance;
-	static TObjectPtr<UFabBrowserApi> JavascriptApi;
-	static TSharedPtr<SDockTab> DockTab;
 	static TUniquePtr<FSlateStyleSet> SlateStyleSet;
-	static TSharedPtr<IWebBrowserWindow> WebBrowserWindow;
 	static TObjectPtr<const UFabSettings> FabPluginSettings;
 
-	static const FName TabId;
+	// Tabs
+	static int32 FabTabCounter;
+	static FName ActiveFabTabId;
+	static TMap<FName, FFabBrowserTabContext> TabsById;
+	static TSharedPtr<FTabManager> TabManager;
+	static bool bIsShuttingDown;
+
 	static const FText FabLabel;
 	static const FText FabTooltip;
 	static const FName FabMenuIconName;
@@ -70,14 +88,32 @@ private:
 	static const FName FabToolbarIconName;
 
 	static void RegisterSlateStyle();
-	static void RegisterNomadTab();
 	static void ExtendContextMenuInContentBrowser();
 	static void SetupEntryPoints();
-	static void ExecuteJavascript(const FString& InSrcScript);
+
+	// Helpers
+	static FName MakeNextFabTabId();
+	static bool EnsureWebBrowserAvailable();
+	static FString BuildIndexFileUrl();
+	static void PreventBrowserKeyConflictWithEditor(const FFabBrowserTabContext& Context);
+	static void InstallUrlRestrictions(const FFabBrowserTabContext& Context);
+	static void InstallDebugPopup(const FFabBrowserTabContext& Context);
 	static void PreventCaptchaFailure();
-	static bool InterceptNavigation(const FString& CurrentUrl, const FString& FabBaseUrl, const FString& FallbackUrl);
-	static TSharedRef<SDockTab> OpenTab(const FSpawnTabArgs& InArgs);
-	static void OnPluginTabClosed(TSharedRef<class SDockTab> InParentTab);
+	static bool InterceptNavigation(const TSharedPtr<IWebBrowserWindow>& BrowserWindow, const FString& CurrentUrl, const FString& FabBaseUrl, const FString& FallbackUrl);
+	static void SetupDockTabListenersForNativeRendering(const TSharedRef<SDockTab>& DockTab, const FName TabId);
+	static void SetupTabWindowMoveListenerForNativeRendering(const TSharedRef<SDockTab>& DockTab);
+	static void ReparentDockTabs();
+	static void RefreshNativeTabViewPort();
+
+	// Tab Creation
+	static void CreateNewFabTab(const FString& Url = TEXT(""));
+	static void CreateOsrFabTab(const FString& Url);
+	static void PreventWebBrowserFreezeInOsrMode(const FName& TabId, const FFabBrowserTabContext& Context);
+	static void CreateNativeFabTab(const FString& Url);
+	static void AttachBrowserToDockTab(const TSharedRef<SDockTab>& DockTab, const bool bNativeRendering, const FName TabId, const FString& Url);
+	static void HandleTabClosed(TSharedRef<SDockTab> ClosedTab, FName TabId);
+	static void HandleTabActivated(TSharedRef<SDockTab> ActivatedTab, ETabActivationCause Cause, FName TabId);
+
 	static TSharedRef<FExtender> OnExtendContentBrowserAssetSelectionMenu(const TArray<FAssetData>& SelectedAssets);
 	static TSharedRef<SWidget> OnFabAssetIconGenerate(const FAssetData& AssetData);
 
@@ -85,12 +121,20 @@ public:
 	static void Init();
 	static void Shutdown();
 
+	/** Notify all open Fab tabs that a logout has occurred by invoking their OnLogout JS callback. */
+	static void NotifyAllTabsLogout();
+
+	/** Notify all open Fab tabs that a login has occurred by sending the access token via their OnLogin JS callback. */
+	static void NotifyAllTabsLogin(const FString& AccessToken);
+
 	static void LogEvent(const FFabAnalyticsPayload& Payload);
 	static void LoggedIn(const FString& InAccessToken);
-	static void GetSignedUrl(const FString& AssetId, const int32 Tier);
-	static TObjectPtr<UFabBrowserApi> GetBrowserApi() { return JavascriptApi; }
+	static TObjectPtr<UFabBrowserApi> GetBrowserApi();
 	static void ShowSettings();
 	static void OpenURL(const FString& InURL = GetUrl());
+	/** Opens the given URL in a new Fab plugin tab (native webview, not the system browser). */
+	static void OpenInNewTab(const FString& Url);
+	static FFabBrowserTabContext* GetActiveDockTabContext();
 	static FString GetUrl();
 	static const ISlateStyle& GetStyleSet();
 };
diff --git a/Engine/Plugins/Fab/Source/Fab/Private/FabBrowserApi.cpp b/Engine/Plugins/Fab/Source/Fab/Private/FabBrowserApi.cpp
--- a/Engine/Plugins/Fab/Source/Fab/Private/FabBrowserApi.cpp
+++ b/Engine/Plugins/Fab/Source/Fab/Private/FabBrowserApi.cpp
@@ -36,6 +36,43 @@ void UFabBrowserApi::CompleteWorkflow(const FString& Id)
 	this->ActiveWorkflows.Remove(CompletedWorkflow);
 }
 
+void UFabBrowserApi::RegisterCallbacks(const FFabCallbacks& Callbacks)
+{
+	OnLogoutDelegate.Unbind();
+	OnLoginDelegate.Unbind();
+	OnLogEventDelegate.Unbind();
+
+	if (Callbacks.OnLogout.IsValid())
+	{
+		OnLogoutDelegate.BindLambda(Callbacks.OnLogout);
+	}
+
+	if (Callbacks.OnLogin.IsValid())
+	{
+		OnLoginDelegate.BindLambda(Callbacks.OnLogin);
+	}
+
+	if (Callbacks.OnLogEvent.IsValid())
+	{
+		OnLogEventDelegate.BindLambda(Callbacks.OnLogEvent);
+	}
+}
+
+void UFabBrowserApi::ExecuteOnLogoutCallback() const
+{
+	OnLogoutDelegate.ExecuteIfBound(0);
+}
+
+void UFabBrowserApi::ExecuteOnLoginCallback(const FString& AccessToken) const
+{
+	OnLoginDelegate.ExecuteIfBound(AccessToken);
+}
+
+void UFabBrowserApi::ExecuteOnLogEventCallback(const FString& JSONPayload) const
+{
+	OnLogEventDelegate.ExecuteIfBound(JSONPayload);
+}
+
 void UFabBrowserApi::AddToProject(const FString& DownloadUrl, const FFabAssetMetadata& AssetMetadata)
 {
 	// Check if the listing is already being downloaded
@@ -337,6 +374,7 @@ void UFabBrowserApi::Login()
 void UFabBrowserApi::Logout()
 {
 	FabAuthentication::DeletePersistentAuth();
+	FFabBrowser::NotifyAllTabsLogout();
 }
 
 void UFabBrowserApi::OpenPluginSettings()
@@ -408,7 +446,7 @@ FFabApiVersion UFabBrowserApi::GetApiVersion()
 	const uint16 MinorVersion = EngineVersion.GetMinor();
 
 	ApiVersion.ue = FString::FromInt(MajorVersion) + "." + FString::FromInt(MinorVersion);
-	ApiVersion.api = "1.0.0";
+	ApiVersion.api = "1.1.0";
 
 	if (const TSharedPtr<IPlugin> Plugin = IPluginManager::Get().FindPlugin("Fab"); Plugin.IsValid())
 	{
@@ -429,6 +467,12 @@ FFabApiVersion UFabBrowserApi::GetApiVersion()
 	return ApiVersion;
 }
 
+void UFabBrowserApi::OpenInNewTab(const FString& Url)
+{
+	FAB_LOG("OpenInNewTab %s", *Url);
+	FFabBrowser::OpenInNewTab(Url);
+}
+
 void UFabBrowserApi::OpenUrlInBrowser(const FString& Url)
 {
 	FAB_LOG("OpenUrlInBrowser %s", *Url);
diff --git a/Engine/Plugins/Fab/Source/Fab/Private/FabBrowserApi.h b/Engine/Plugins/Fab/Source/Fab/Private/FabBrowserApi.h
--- a/Engine/Plugins/Fab/Source/Fab/Private/FabBrowserApi.h
+++ b/Engine/Plugins/Fab/Source/Fab/Private/FabBrowserApi.h
@@ -3,6 +3,7 @@
 #pragma once
 
 #include "CoreMinimal.h"
+#include "WebJSFunction.h"
 
 #include "Workflows/FabWorkflow.h"
 
@@ -27,7 +28,7 @@ struct FFabApiVersion
 	GENERATED_BODY()
 
 	// These uproperties need to be in this case (NOT PascalCase)
-	// otherwise we will break compatibility with the engine plugin 
+	// otherwise we will break compatibility with the engine plugin
 	// Fab's JS on the macOS WebKit WebBrowser plugin backend.
 
 	UPROPERTY()
@@ -41,7 +42,7 @@ struct FFabApiVersion
 
 	UPROPERTY()
 	FString platform;
-	
+
 	UPROPERTY()
 	FFabSupportedFeatures supportedfeatures;
 };
@@ -52,7 +53,7 @@ struct FFabFrontendSettings
 	GENERATED_BODY()
 
 	// These two uproperties need to be in this case (NOT PascalCase)
-	// otherwise we will break compatibility with the engine plugin 
+	// otherwise we will break compatibility with the engine plugin
 	// Fab's JS on the macOS WebKit WebBrowser plugin backend.
 	UPROPERTY()
 	FString preferredformat;
@@ -61,20 +62,54 @@ struct FFabFrontendSettings
 	FString preferredquality;
 };
 
+USTRUCT()
+struct FFabCallbacks
+{
+	GENERATED_BODY()
+
+	UPROPERTY()
+	FWebJSFunction OnLogout;
+
+	UPROPERTY()
+	FWebJSFunction OnLogin;
+
+	UPROPERTY()
+	FWebJSFunction OnLogEvent;
+};
+
 UCLASS()
 class UFabBrowserApi : public UObject
 {
 	GENERATED_BODY()
 	DECLARE_MULTICAST_DELEGATE_TwoParams(FOnSignedUrlGenerated, const FString& /*DownloadUrl*/, FFabAssetMetadata /*Metadata*/);
 
+	// Frontend callbacks
+	DECLARE_DELEGATE_OneParam(FOnLogout, int32 /* dummy value */);
+	DECLARE_DELEGATE_OneParam(FOnLogin, const FString& /* AccessToken */);
+	DECLARE_DELEGATE_OneParam(FOnLogEvent, const FString& /* JSONPayload */);
+
 private:
 	FOnSignedUrlGenerated OnSignedUrlGeneratedDelegate;
 	void CompleteWorkflow(const FString& Id);
-
-public:
 	TArray<TSharedPtr<IFabWorkflow>> ActiveWorkflows;
 
+	FOnLogout OnLogoutDelegate;
+	FOnLogin OnLoginDelegate;
+	FOnLogEvent OnLogEventDelegate;
+
 public:
+	UFUNCTION()
+	void RegisterCallbacks(const FFabCallbacks& Callbacks);
+
+	/** Invoke the stored OnLogout callback, if one has been registered by the frontend. */
+	void ExecuteOnLogoutCallback() const;
+
+	/** Invoke the stored OnLogin callback, if one has been registered by the frontend. */
+	void ExecuteOnLoginCallback(const FString& AccessToken) const;
+
+	/** Invoke the stored OnLogEvent callback, if one has been registered by the frontend. */
+	void ExecuteOnLogEventCallback(const FString& JSONPayload) const;
+
 	UFUNCTION()
 	void AddToProject(const FString& DownloadUrl, const FFabAssetMetadata& AssetMetadata);
 
@@ -117,6 +152,9 @@ public:
 	UFUNCTION()
 	static FFabApiVersion GetApiVersion();
 
+	UFUNCTION()
+	void OpenInNewTab(const FString& Url);
+
 	UFUNCTION()
 	void OpenUrlInBrowser(const FString& Url);
 
diff --git a/Engine/Plugins/Fab/Source/Fab/Private/FabSettings.h b/Engine/Plugins/Fab/Source/Fab/Private/FabSettings.h
--- a/Engine/Plugins/Fab/Source/Fab/Private/FabSettings.h
+++ b/Engine/Plugins/Fab/Source/Fab/Private/FabSettings.h
@@ -49,10 +49,14 @@ public:
 	UPROPERTY(config, EditAnywhere, Category=Frontend, meta=(DevOnly=true))
 	FString CustomUrl;
 
-	/** Custom auth token used when it's non empty */
+	/** Custom auth token used when it's non-empty */
 	UPROPERTY(config, EditAnywhere, Category=Frontend, meta=(DevOnly=true))
 	FString CustomAuthToken;
 
+	/** Enables native rendering mode (Windows only; Mac always uses native rendering with WebKit) */
+	UPROPERTY(config, EditAnywhere, Category=General)
+	bool bEnableNativeRendering = false;
+
 	/** Enable chrome debug options - default is false */
 	UPROPERTY(config, EditAnywhere, Category = General)
 	bool bEnableDebugOptions = false;
diff --git a/Engine/Plugins/Fab/Source/Fab/Private/FabSettingsCustomization.cpp b/Engine/Plugins/Fab/Source/Fab/Private/FabSettingsCustomization.cpp
--- a/Engine/Plugins/Fab/Source/Fab/Private/FabSettingsCustomization.cpp
+++ b/Engine/Plugins/Fab/Source/Fab/Private/FabSettingsCustomization.cpp
@@ -29,6 +29,11 @@ void FFabSettingsCustomization::CustomizeDetails(IDetailLayoutBuilder& DetailBui
 	// General section modifications
 
 	DetailBuilder.HideProperty(GET_MEMBER_NAME_CHECKED(UFabSettings, CacheDirectorySize));
+
+#if PLATFORM_MAC
+	// On Mac, WebKit is always natively rendered, so hide this Windows-only setting
+	DetailBuilder.HideProperty(GET_MEMBER_NAME_CHECKED(UFabSettings, bEnableNativeRendering));
+#endif
 	TArray<TSharedRef<IPropertyHandle>> GeneralProperties;
 	GeneralCategory.GetDefaultProperties(GeneralProperties);
 	for (const TSharedRef<IPropertyHandle>& PropertyHandle : GeneralProperties)
diff --git a/Engine/Plugins/Fab/Source/Fab/Private/Utilities/DragImportOperation.cpp b/Engine/Plugins/Fab/Source/Fab/Private/Utilities/DragImportOperation.cpp
--- a/Engine/Plugins/Fab/Source/Fab/Private/Utilities/DragImportOperation.cpp
+++ b/Engine/Plugins/Fab/Source/Fab/Private/Utilities/DragImportOperation.cpp
@@ -2,9 +2,9 @@
 
 #include "DragImportOperation.h"
 
-#include "Editor.h"
+#include "FabBrowser.h"
 #include "FabLog.h"
-#include "Materials/MaterialParameters.h"
+#include "IWebBrowserWindow.h"
 
 #include "Animation/SkeletalMeshActor.h"
 #include "Components/DecalComponent.h"
@@ -13,6 +13,7 @@
 #include "Components/SkeletalMeshComponent.h"
 
 #include "Engine/DecalActor.h"
+#include "Engine/StaticMesh.h"
 #include "Engine/StaticMeshActor.h"
 #include "Engine/SkeletalMesh.h"
 
@@ -74,7 +75,27 @@ void FDragImportOperation::InitializeDrag()
 
 	FDragDropEvent DragDropEvent(FakePointerEvent, DragOperationHandle);
 
-	FSlateApplication::Get().ProcessDragEnterEvent(FSlateApplication::Get().GetTopLevelWindows()[0], DragDropEvent);
+	FFabBrowserTabContext* const& FabBrowserTabContext = FFabBrowser::GetActiveDockTabContext();
+
+	if (FabBrowserTabContext && FabBrowserTabContext->WebBrowserWindow.IsValid())
+	{
+		TSharedPtr<SWindow> ParentWindow = FabBrowserTabContext->WebBrowserWindow->GetParentWindow();
+		FSlateApplication::Get().ProcessDragEnterEvent(ParentWindow.ToSharedRef(), DragDropEvent);
+
+		// In native (non-OSR) rendering mode, CEF owns a real child HWND and holds OS
+		// mouse capture from the original WM_LBUTTONDOWN. That means WM_LBUTTONUP is
+		// delivered to CEF rather than Slate, so Slate never sees the button release and
+		// stays in drag-drop mode until the next unrelated click.
+		// Transferring capture to the parent Slate window here ensures WM_LBUTTONUP
+		// goes through Slate's own input pipeline and ends the drag naturally.
+		if (FabBrowserTabContext->bIsNativelyRendered)
+		{
+			if (TSharedPtr<FGenericWindow> NativeWindow = ParentWindow->GetNativeWindow())
+			{
+				FSlateApplication::Get().GetPlatformApplication()->SetCapture(NativeWindow);
+			}
+		}
+	}
 }
 
 void FDragImportOperation::UpdateDraggedAsset(UObject* InDraggedObject, const EDragAssetType InDragAssetType)
diff --git a/Engine/Plugins/Fab/Source/Fab/Private/Workflows/GenericDragDropWorkflow.cpp b/Engine/Plugins/Fab/Source/Fab/Private/Workflows/GenericDragDropWorkflow.cpp
--- a/Engine/Plugins/Fab/Source/Fab/Private/Workflows/GenericDragDropWorkflow.cpp
+++ b/Engine/Plugins/Fab/Source/Fab/Private/Workflows/GenericDragDropWorkflow.cpp
@@ -40,7 +40,15 @@ void FGenericDragDropWorkflow::Execute()
 	FAssetData PlaceHolderAsset(FSoftObjectPath("/Fab/Placeholders/MeshPlaceholder.MeshPlaceholder").TryLoad());
 	DragOperation = MakeUnique<FDragImportOperation>(PlaceHolderAsset, EDragAssetType::Mesh);
 
-	auto OnDragInfoSuccess = [this](const FString& InDownloadUrl, const FFabAssetMetadata& AssetMetadata)
+	const TObjectPtr<UFabBrowserApi>& FabBrowserApi = FFabBrowser::GetBrowserApi();
+	if (!FabBrowserApi)
+	{
+		return;
+	}
+
+	TWeakObjectPtr<UFabBrowserApi> WeakApi = FabBrowserApi;
+
+	auto OnDragInfoSuccess = [this, WeakApi](const FString& InDownloadUrl, const FFabAssetMetadata& AssetMetadata)
 	{
 		if (AssetId == AssetMetadata.AssetId)
 		{
@@ -55,13 +63,16 @@ void FGenericDragDropWorkflow::Execute()
 			}
 			if (SignedUrlHandle.IsValid())
 			{
-				FFabBrowser::GetBrowserApi()->RemoveSignedUrlHandle(SignedUrlHandle);
+				if (UFabBrowserApi* Api = WeakApi.Get())
+				{
+					Api->RemoveSignedUrlHandle(SignedUrlHandle);
+				}
 				SignedUrlHandle.Reset();
 			}
 		}
 	};
 
-	SignedUrlHandle = FFabBrowser::GetBrowserApi()->AddSignedUrlCallback(OnDragInfoSuccess);
+	SignedUrlHandle = FabBrowserApi->AddSignedUrlCallback(OnDragInfoSuccess);
 }
 
 void FGenericDragDropWorkflow::OnContentDownloadProgress(const FFabDownloadRequest* Request, const FFabDownloadStats& DownloadStats)
diff --git a/Engine/Plugins/Fab/Source/Fab/Private/Workflows/QuixelDragDropWorkflow.cpp b/Engine/Plugins/Fab/Source/Fab/Private/Workflows/QuixelDragDropWorkflow.cpp
--- a/Engine/Plugins/Fab/Source/Fab/Private/Workflows/QuixelDragDropWorkflow.cpp
+++ b/Engine/Plugins/Fab/Source/Fab/Private/Workflows/QuixelDragDropWorkflow.cpp
@@ -93,7 +93,15 @@ void FQuixelDragDropWorkflow::Execute()
 
 	DragOperation = MakeUnique<FDragImportOperation>(PlaceHolderAsset, DragAssetType);
 
-	auto OnDragInfoSuccess = [this](const FString& InDownloadUrl, const FFabAssetMetadata& AssetMetadata)
+	const TObjectPtr<UFabBrowserApi>& FabBrowserApi = FFabBrowser::GetBrowserApi();
+	if (!FabBrowserApi)
+	{
+		return;
+	}
+
+	TWeakObjectPtr<UFabBrowserApi> WeakApi = FabBrowserApi;
+
+	auto OnDragInfoSuccess = [this, WeakApi](const FString& InDownloadUrl, const FFabAssetMetadata& AssetMetadata)
 	{
 		if (AssetId == AssetMetadata.AssetId)
 		{
@@ -108,13 +116,16 @@ void FQuixelDragDropWorkflow::Execute()
 			}
 			if (SignedUrlHandle.IsValid())
 			{
-				FFabBrowser::GetBrowserApi()->RemoveSignedUrlHandle(SignedUrlHandle);
+				if (UFabBrowserApi* Api = WeakApi.Get())
+				{
+					Api->RemoveSignedUrlHandle(SignedUrlHandle);
+				}
 				SignedUrlHandle.Reset();
 			}
 		}
 	};
 
-	SignedUrlHandle = FFabBrowser::GetBrowserApi()->AddSignedUrlCallback(OnDragInfoSuccess);
+	SignedUrlHandle = FabBrowserApi->AddSignedUrlCallback(OnDragInfoSuccess);
 }
 
 void FQuixelDragDropWorkflow::OnContentDownloadProgress(const FFabDownloadRequest* Request, const FFabDownloadStats& DownloadStats)
