Fab: Override user agent for the browser instead of globally

Use the new per-browser user agent override setting, so we don't override all the other browsers unintentionally. It would only work if Fab was the first created browser anyway.

#rnx
#Synced-CL 52790326

diff --git a/Engine/Plugins/Fab/Source/Fab/Private/FabBrowser.cpp b/Engine/Plugins/Fab/Source/Fab/Private/FabBrowser.cpp
--- a/Engine/Plugins/Fab/Source/Fab/Private/FabBrowser.cpp
+++ b/Engine/Plugins/Fab/Source/Fab/Private/FabBrowser.cpp
@@ -95,7 +95,6 @@ static TSharedPtr<FFabWindowMoveMessageHandler> WindowMoveMessageHandler;
 void FFabBrowser::Init()
 {
 	bIsShuttingDown = false;
-	PreventCaptchaFailure();
 	RegisterSlateStyle();
 	SetupEntryPoints();
 	ExtendContextMenuInContentBrowser();
@@ -526,7 +525,10 @@ void FFabBrowser::AttachBrowserToDockTab(const TSharedRef<SDockTab>& DockTab, co
 	// Browser settings
 	FCreateBrowserWindowSettings Settings;
 	Settings.InitialURL = Url.IsEmpty() ? BuildIndexFileUrl() : Url;
-	Settings.BrowserFrameRate = 60;
+	// Override the user agent to avoid issues with the project name
+	// If the project name has the word "agent" in it, it'll cause the captcha to fail
+	Settings.UserAgentApplication = IWebBrowserModule::MakeUserAgentApplication(TEXT("Fab"));
+
 #if PLATFORM_MAC
 	Settings.bMobileJSReturnInDict = false;
 #endif
@@ -1309,14 +1311,4 @@ bool FFabBrowser::InterceptNavigation(const TSharedPtr<IWebBrowserWindow>& Brows
 	return false;
 }
 
-void FFabBrowser::PreventCaptchaFailure()
-{
-	// Override the user agent to avoid issues with the project name
-	// If the project name has the word "agent" in it, it'll cause the captcha to fail
-	FWebBrowserInitSettings BrowserInitSettings = FWebBrowserInitSettings();
-	FString VersionString = FEngineVersion::Current().ToString();
-	BrowserInitSettings.ProductVersion = VersionString;
-	IWebBrowserModule::Get().CustomInitialize(BrowserInitSettings);
-}
-
 #undef LOCTEXT_NAMESPACE
diff --git a/Engine/Plugins/Fab/Source/Fab/Private/FabBrowser.h b/Engine/Plugins/Fab/Source/Fab/Private/FabBrowser.h
--- a/Engine/Plugins/Fab/Source/Fab/Private/FabBrowser.h
+++ b/Engine/Plugins/Fab/Source/Fab/Private/FabBrowser.h
@@ -98,7 +98,6 @@ private:
 	static void PreventBrowserKeyConflictWithEditor(const FFabBrowserTabContext& Context);
 	static void InstallUrlRestrictions(const FFabBrowserTabContext& Context);
 	static void InstallDebugPopup(const FFabBrowserTabContext& Context);
-	static void PreventCaptchaFailure();
 	static bool InterceptNavigation(const TSharedPtr<IWebBrowserWindow>& BrowserWindow, const FString& CurrentUrl, const FString& FabBaseUrl, const FString& FallbackUrl);
 	static void SetupDockTabListenersForNativeRendering(const TSharedRef<SDockTab>& DockTab, const FName TabId);
 	static void SetupTabWindowMoveListenerForNativeRendering(const TSharedRef<SDockTab>& DockTab);
